Add deprecation messages and update samples

diff --git a/Engine/Plugins/Lumin/MagicLeap/Source/MagicLeapController/Private/MagicLeapController.cpp b/Engine/Plugins/Lumin/MagicLeap/Source/MagicLeapController/Private/MagicLeapController.cpp
--- a/Engine/Plugins/Lumin/MagicLeap/Source/MagicLeapController/Private/MagicLeapController.cpp
+++ b/Engine/Plugins/Lumin/MagicLeap/Source/MagicLeapController/Private/MagicLeapController.cpp
@@ -1164,7 +1164,7 @@ void FMagicLeapController::AddKeys()
 {
 	EKeys::AddMenuCategoryDisplayInfo("MagicLeap", LOCTEXT("MagicLeapSubCategory", "Magic Leap"), TEXT("GraphEditor.PadEvent_16x"));
 
-	EKeys::AddKey(FKeyDetails(FMagicLeapKeys::MotionController_Left_Thumbstick_Z, LOCTEXT("MotionController_Left_Thumbstick_Z", "MotionController (L) Thumbstick Z"), FKeyDetails::GamepadKey | FKeyDetails::FloatAxis | FKeyDetails::NotBlueprintBindableKey | FKeyDetails::NotActionBindableKey));
+	EKeys::AddKey(FKeyDetails(FMagicLeapKeys::MotionController_Left_Thumbstick_Z, LOCTEXT("MotionController_Left_Thumbstick_Z", "MotionController (L) Thumbstick Z"), FKeyDetails::GamepadKey | FKeyDetails::FloatAxis | FKeyDetails::Deprecated));
 	EKeys::AddKey(FKeyDetails(FMagicLeapKeys::Left_HomeButton, LOCTEXT("MagicLeap_Left_HomeButton", "ML (L) Home Button"), FKeyDetails::GamepadKey | FKeyDetails::NotBlueprintBindableKey, "MagicLeap"));
 	EKeys::AddKey(FKeyDetails(FMagicLeapKeys::Left_Bumper, LOCTEXT("MagicLeap_Left_Bumper", "ML (L) Bumper"), FKeyDetails::GamepadKey | FKeyDetails::NotBlueprintBindableKey, "MagicLeap"));
 	EKeys::AddKey(FKeyDetails(FMagicLeapKeys::Left_Trigger, LOCTEXT("MagicLeap_Left_Trigger", "ML (L) Trigger"), FKeyDetails::GamepadKey | FKeyDetails::NotBlueprintBindableKey, "MagicLeap"));
@@ -1178,7 +1178,7 @@ void FMagicLeapController::AddKeys()
 	EKeys::AddKey(FKeyDetails(FMagicLeapKeys::Left_Touch1_Force, LOCTEXT("MagicLeap_Left_Touch1_Force", "ML (L) Touch1 Force"), FKeyDetails::GamepadKey | FKeyDetails::FloatAxis | FKeyDetails::NotBlueprintBindableKey, "MagicLeap"));
 	EKeys::AddKey(FKeyDetails(FMagicLeapKeys::Left_Touch1_Touch, LOCTEXT("MagicLeap_Left_Touch1_Touch", "ML (L) Touch1 Touch"), FKeyDetails::GamepadKey | FKeyDetails::NotBlueprintBindableKey, "MagicLeap"));
 
-	EKeys::AddKey(FKeyDetails(FMagicLeapKeys::MotionController_Right_Thumbstick_Z, LOCTEXT("MotionController_Right_Thumbstick_Z", "MotionController (R) Thumbstick Z"), FKeyDetails::GamepadKey | FKeyDetails::FloatAxis | FKeyDetails::NotBlueprintBindableKey | FKeyDetails::NotActionBindableKey));
+	EKeys::AddKey(FKeyDetails(FMagicLeapKeys::MotionController_Right_Thumbstick_Z, LOCTEXT("MotionController_Right_Thumbstick_Z", "MotionController (R) Thumbstick Z"), FKeyDetails::GamepadKey | FKeyDetails::FloatAxis | FKeyDetails::Deprecated));
 	EKeys::AddKey(FKeyDetails(FMagicLeapKeys::Right_HomeButton, LOCTEXT("MagicLeap_Right_HomeButton", "ML (R) Home Button"), FKeyDetails::GamepadKey | FKeyDetails::NotBlueprintBindableKey, "MagicLeap"));
 	EKeys::AddKey(FKeyDetails(FMagicLeapKeys::Right_Bumper, LOCTEXT("MagicLeap_Right_Bumper", "ML (R) Bumper"), FKeyDetails::GamepadKey | FKeyDetails::NotBlueprintBindableKey, "MagicLeap"));
 	EKeys::AddKey(FKeyDetails(FMagicLeapKeys::Right_Trigger, LOCTEXT("MagicLeap_Right_Trigger", "ML (R) Trigger"), FKeyDetails::GamepadKey | FKeyDetails::NotBlueprintBindableKey, "MagicLeap"));
diff --git a/Engine/Plugins/Runtime/Oculus/OculusVR/Source/OculusInput/Private/OculusInput.cpp b/Engine/Plugins/Runtime/Oculus/OculusVR/Source/OculusInput/Private/OculusInput.cpp
--- a/Engine/Plugins/Runtime/Oculus/OculusVR/Source/OculusInput/Private/OculusInput.cpp
+++ b/Engine/Plugins/Runtime/Oculus/OculusVR/Source/OculusInput/Private/OculusInput.cpp
@@ -125,10 +125,10 @@ void FOculusInput::PreInit()
 	LoadConfig();
 
 	// Register the FKeys
-	EKeys::AddKey(FKeyDetails(FOculusKey::OculusTouch_Left_Thumbstick, LOCTEXT("OculusTouch_Left_Thumbstick", "Oculus Touch (L) Thumbstick CapTouch"), FKeyDetails::GamepadKey | FKeyDetails::FloatAxis | FKeyDetails::NotBlueprintBindableKey | FKeyDetails::NotActionBindableKey));
-	EKeys::AddKey(FKeyDetails(FOculusKey::OculusTouch_Left_FaceButton1, LOCTEXT("OculusTouch_Left_FaceButton1", "Oculus Touch (L) X Button CapTouch"), FKeyDetails::GamepadKey | FKeyDetails::FloatAxis | FKeyDetails::NotBlueprintBindableKey | FKeyDetails::NotActionBindableKey));
-	EKeys::AddKey(FKeyDetails(FOculusKey::OculusTouch_Left_Trigger, LOCTEXT("OculusTouch_Left_Trigger", "Oculus Touch (L) Trigger CapTouch"), FKeyDetails::GamepadKey | FKeyDetails::FloatAxis | FKeyDetails::NotBlueprintBindableKey | FKeyDetails::NotActionBindableKey));
-	EKeys::AddKey(FKeyDetails(FOculusKey::OculusTouch_Left_FaceButton2, LOCTEXT("OculusTouch_Left_FaceButton2", "Oculus Touch (L) Y Button CapTouch"), FKeyDetails::GamepadKey | FKeyDetails::FloatAxis | FKeyDetails::NotBlueprintBindableKey | FKeyDetails::NotActionBindableKey));
+	EKeys::AddKey(FKeyDetails(FOculusKey::OculusTouch_Left_Thumbstick, LOCTEXT("OculusTouch_Left_Thumbstick", "Oculus Touch (L) Thumbstick CapTouch"), FKeyDetails::GamepadKey | FKeyDetails::FloatAxis | FKeyDetails::Deprecated));
+	EKeys::AddKey(FKeyDetails(FOculusKey::OculusTouch_Left_FaceButton1, LOCTEXT("OculusTouch_Left_FaceButton1", "Oculus Touch (L) X Button CapTouch"), FKeyDetails::GamepadKey | FKeyDetails::FloatAxis | FKeyDetails::Deprecated));
+	EKeys::AddKey(FKeyDetails(FOculusKey::OculusTouch_Left_Trigger, LOCTEXT("OculusTouch_Left_Trigger", "Oculus Touch (L) Trigger CapTouch"), FKeyDetails::GamepadKey | FKeyDetails::FloatAxis | FKeyDetails::Deprecated));
+	EKeys::AddKey(FKeyDetails(FOculusKey::OculusTouch_Left_FaceButton2, LOCTEXT("OculusTouch_Left_FaceButton2", "Oculus Touch (L) Y Button CapTouch"), FKeyDetails::GamepadKey | FKeyDetails::FloatAxis | FKeyDetails::Deprecated));
 	EKeys::AddKey(FKeyDetails(FOculusKey::OculusTouch_Left_IndexPointing, LOCTEXT("OculusTouch_Left_IndexPointing", "Oculus Touch (L) Pointing CapTouch"), FKeyDetails::GamepadKey | FKeyDetails::FloatAxis | FKeyDetails::NotBlueprintBindableKey, "OculusTouch"));
 	EKeys::AddKey(FKeyDetails(FOculusKey::OculusTouch_Left_ThumbUp, LOCTEXT("OculusTouch_Left_ThumbUp", "Oculus Touch (L) Thumb Up CapTouch"), FKeyDetails::GamepadKey | FKeyDetails::FloatAxis | FKeyDetails::NotBlueprintBindableKey, "OculusTouch"));
 
@@ -153,10 +153,10 @@ void FOculusInput::PreInit()
 	EKeys::AddKey(FKeyDetails(FOculusKey::OculusRemote_VolumeDown, LOCTEXT("OculusRemote_VolumeDown", "Oculus Remote Volume Down"), FKeyDetails::GamepadKey, "OculusRemote"));
 	EKeys::AddKey(FKeyDetails(FOculusKey::OculusRemote_Home, LOCTEXT("OculusRemote_Home", "Oculus Remote Home"), FKeyDetails::GamepadKey, "OculusRemote"));
 
-	EKeys::AddKey(FKeyDetails(FOculusKey::OculusTouchpad_Touchpad, LOCTEXT("OculusTouchpad_Touchpad", "Oculus Touchpad Button"), FKeyDetails::GamepadKey | FKeyDetails::NotBlueprintBindableKey | FKeyDetails::NotActionBindableKey));
-	EKeys::AddKey(FKeyDetails(FOculusKey::OculusTouchpad_Touchpad_X, LOCTEXT("OculusTouchpad_Touchpad_X", "Oculus Touchpad X-Axis"), FKeyDetails::GamepadKey | FKeyDetails::FloatAxis | FKeyDetails::NotBlueprintBindableKey | FKeyDetails::NotActionBindableKey));
-	EKeys::AddKey(FKeyDetails(FOculusKey::OculusTouchpad_Touchpad_Y, LOCTEXT("OculusTouchpad_Touchpad_Y", "Oculus Touchpad Y-Axis"), FKeyDetails::GamepadKey | FKeyDetails::FloatAxis | FKeyDetails::NotBlueprintBindableKey | FKeyDetails::NotActionBindableKey));
-	EKeys::AddKey(FKeyDetails(FOculusKey::OculusTouchpad_Back, LOCTEXT("OculusTouchpad_Back", "Oculus Touchpad Back"), FKeyDetails::GamepadKey | FKeyDetails::NotBlueprintBindableKey | FKeyDetails::NotActionBindableKey));
+	EKeys::AddKey(FKeyDetails(FOculusKey::OculusTouchpad_Touchpad, LOCTEXT("OculusTouchpad_Touchpad", "Oculus Touchpad Button"), FKeyDetails::GamepadKey | FKeyDetails::Deprecated));
+	EKeys::AddKey(FKeyDetails(FOculusKey::OculusTouchpad_Touchpad_X, LOCTEXT("OculusTouchpad_Touchpad_X", "Oculus Touchpad X-Axis"), FKeyDetails::GamepadKey | FKeyDetails::FloatAxis | FKeyDetails::Deprecated));
+	EKeys::AddKey(FKeyDetails(FOculusKey::OculusTouchpad_Touchpad_Y, LOCTEXT("OculusTouchpad_Touchpad_Y", "Oculus Touchpad Y-Axis"), FKeyDetails::GamepadKey | FKeyDetails::FloatAxis | FKeyDetails::Deprecated));
+	EKeys::AddKey(FKeyDetails(FOculusKey::OculusTouchpad_Back, LOCTEXT("OculusTouchpad_Back", "Oculus Touchpad Back"), FKeyDetails::GamepadKey | FKeyDetails::Deprecated));
 
 	UE_LOG(LogOcInput, Log, TEXT("OculusInput pre-init called"));
 }
diff --git a/Engine/Source/Editor/BlueprintGraph/Private/K2Node_GetInputAxisKeyValue.cpp b/Engine/Source/Editor/BlueprintGraph/Private/K2Node_GetInputAxisKeyValue.cpp
--- a/Engine/Source/Editor/BlueprintGraph/Private/K2Node_GetInputAxisKeyValue.cpp
+++ b/Engine/Source/Editor/BlueprintGraph/Private/K2Node_GetInputAxisKeyValue.cpp
@@ -86,6 +86,10 @@ void UK2Node_GetInputAxisKeyValue::ValidateNodeDuringCompilation(class FCompiler
 	{
 		MessageLog.Warning(*FText::Format(NSLOCTEXT("KismetCompiler", "NotAxis_GetInputAxisKey_Warning", "GetInputAxisKey Value specifies FKey'{0}' which is not a float axis for @@"), FText::FromString(InputAxisKey.ToString())).ToString(), this);
 	}
+	else if (InputAxisKey.IsDeprecated())
+	{
+		MessageLog.Warning(*FText::Format(NSLOCTEXT("KismetCompiler", "Deprecated_GetInputAxisKey_Warning", "GetInputAxisKey Value specifies FKey'{0}' which has been deprecated for @@"), FText::FromString(InputAxisKey.ToString())).ToString(), this);
+	}
 	else if (!InputAxisKey.IsBindableInBlueprints())
 	{
 		MessageLog.Warning(*FText::Format(NSLOCTEXT("KismetCompiler", "NotBindanble_GetInputAxisKey_Warning", "GetInputAxisKey Value specifies FKey'{0}' that is not blueprint bindable for @@"), FText::FromString(InputAxisKey.ToString())).ToString(), this);
diff --git a/Engine/Source/Editor/BlueprintGraph/Private/K2Node_GetInputVectorAxisValue.cpp b/Engine/Source/Editor/BlueprintGraph/Private/K2Node_GetInputVectorAxisValue.cpp
--- a/Engine/Source/Editor/BlueprintGraph/Private/K2Node_GetInputVectorAxisValue.cpp
+++ b/Engine/Source/Editor/BlueprintGraph/Private/K2Node_GetInputVectorAxisValue.cpp
@@ -40,6 +40,10 @@ void UK2Node_GetInputVectorAxisValue::ValidateNodeDuringCompilation(class FCompi
 	{
 		MessageLog.Warning(*FText::Format(NSLOCTEXT("KismetCompiler", "NotAxis_GetInputVectorAxis_Warning", "GetInputVectorAxis Value specifies FKey'{0}' which is not a vector axis for @@"), FText::FromString(InputAxisKey.ToString())).ToString(), this);
 	}
+	else if (InputAxisKey.IsDeprecated())
+	{
+		MessageLog.Warning(*FText::Format(NSLOCTEXT("KismetCompiler", "Deprecated_GetInputVectorAxis_Warning", "GetInputVectorAxis Value specifies FKey'{0}' which has been deprecated for @@"), FText::FromString(InputAxisKey.ToString())).ToString(), this);
+	}
 	else if (!InputAxisKey.IsBindableInBlueprints())
 	{
 		MessageLog.Warning(*FText::Format(NSLOCTEXT("KismetCompiler", "NotBindanble_GetInputVectorAxis_Warning", "GetInputVectorAxis Value specifies FKey'{0}' that is not blueprint bindable for @@"), FText::FromString(InputAxisKey.ToString())).ToString(), this);
diff --git a/Engine/Source/Editor/BlueprintGraph/Private/K2Node_InputAxisKeyEvent.cpp b/Engine/Source/Editor/BlueprintGraph/Private/K2Node_InputAxisKeyEvent.cpp
--- a/Engine/Source/Editor/BlueprintGraph/Private/K2Node_InputAxisKeyEvent.cpp
+++ b/Engine/Source/Editor/BlueprintGraph/Private/K2Node_InputAxisKeyEvent.cpp
@@ -67,6 +67,10 @@ void UK2Node_InputAxisKeyEvent::ValidateNodeDuringCompilation(class FCompilerRes
 	{
 		MessageLog.Warning(*FText::Format(NSLOCTEXT("KismetCompiler", "NotAxis_InputAxisKey_Warning", "InputAxisKey Event specifies FKey'{0}' which is not a float axis for @@"), FText::FromString(AxisKey.ToString())).ToString(), this);
 	}
+	else if (AxisKey.IsDeprecated())
+	{
+		MessageLog.Warning(*FText::Format(NSLOCTEXT("KismetCompiler", "Deprecated_InputAxisKey_Warning", "InputAxisKey Event specifies FKey'{0}' which has been deprecated for @@"), FText::FromString(AxisKey.ToString())).ToString(), this);
+	}
 	else if (!AxisKey.IsBindableInBlueprints())
 	{
 		MessageLog.Warning(*FText::Format(NSLOCTEXT("KismetCompiler", "NotBindable_InputAxisKey_Warning", "InputAxisKey Event specifies FKey'{0}' that is not blueprint bindable for @@"), FText::FromString(AxisKey.ToString())).ToString(), this);
diff --git a/Engine/Source/Editor/BlueprintGraph/Private/K2Node_InputKey.cpp b/Engine/Source/Editor/BlueprintGraph/Private/K2Node_InputKey.cpp
--- a/Engine/Source/Editor/BlueprintGraph/Private/K2Node_InputKey.cpp
+++ b/Engine/Source/Editor/BlueprintGraph/Private/K2Node_InputKey.cpp
@@ -237,6 +237,10 @@ void UK2Node_InputKey::ValidateNodeDuringCompilation(class FCompilerResultsLog&
 	{
 		MessageLog.Warning(*FText::Format(NSLOCTEXT("KismetCompiler", "Axis_InputKey_Warning", "InputKey Event specifies axis FKey'{0}' for @@"), FText::FromString(InputKey.ToString())).ToString(), this);
 	}
+	else if (InputKey.IsDeprecated())
+	{
+		MessageLog.Warning(*FText::Format(NSLOCTEXT("KismetCompiler", "Deprecated_InputKey_Warning", "InputKey Event specifies FKey'{0}' which has been deprecated for @@"), FText::FromString(InputKey.ToString())).ToString(), this);
+	}
 	else if (!InputKey.IsBindableInBlueprints())
 	{
 		MessageLog.Warning( *FText::Format( NSLOCTEXT("KismetCompiler", "NotBindanble_InputKey_Warning", "InputKey Event specifies FKey'{0}' that is not blueprint bindable for @@"), FText::FromString(InputKey.ToString())).ToString(), this);
diff --git a/Engine/Source/Editor/BlueprintGraph/Private/K2Node_InputVectorAxisEvent.cpp b/Engine/Source/Editor/BlueprintGraph/Private/K2Node_InputVectorAxisEvent.cpp
--- a/Engine/Source/Editor/BlueprintGraph/Private/K2Node_InputVectorAxisEvent.cpp
+++ b/Engine/Source/Editor/BlueprintGraph/Private/K2Node_InputVectorAxisEvent.cpp
@@ -39,6 +39,10 @@ void UK2Node_InputVectorAxisEvent::ValidateNodeDuringCompilation(class FCompiler
 	{
 		MessageLog.Warning(*FText::Format(NSLOCTEXT("KismetCompiler", "NotAxis_InputVectorAxis_Warning", "InputVectorAxis Event specifies FKey'{0}' which is not a vector axis for @@"), FText::FromString(AxisKey.ToString())).ToString(), this);
 	}
+	else if (AxisKey.IsDeprecated())
+	{
+		MessageLog.Warning(*FText::Format(NSLOCTEXT("KismetCompiler", "Deprecated_InputVectorAxis_Warning", "InputVectorAxis Event specifies FKey'{0}' which has been deprecated for @@"), FText::FromString(AxisKey.ToString())).ToString(), this);
+	}
 	else if (!AxisKey.IsBindableInBlueprints())
 	{
 		MessageLog.Warning(*FText::Format(NSLOCTEXT("KismetCompiler", "NotBindable_InputVectorAxis_Warning", "InputVectorAxis Event specifies FKey'{0}' that is not blueprint bindable for @@"), FText::FromString(AxisKey.ToString())).ToString(), this);
diff --git a/Engine/Source/Editor/UnrealEd/Private/SKeySelector.cpp b/Engine/Source/Editor/UnrealEd/Private/SKeySelector.cpp
--- a/Engine/Source/Editor/UnrealEd/Private/SKeySelector.cpp
+++ b/Engine/Source/Editor/UnrealEd/Private/SKeySelector.cpp
@@ -125,7 +125,7 @@ void SKeySelector::Construct(const FArguments& InArgs)
 		SAssignNew(KeyComboButton, SComboButton)
 		.OnGetMenuContent(this, &SKeySelector::GetMenuContent)
 		.ContentPadding(0)
-		.ToolTipText(LOCTEXT("KeySelector", "Select the key value."))
+		.ToolTipText(this, &SKeySelector::GetKeyTooltip)
 		.ButtonContent()
 		[
 			SNew(SHorizontalBox)
@@ -167,11 +167,37 @@ const FSlateBrush* SKeySelector::GetKeyIconImage() const
 	TOptional<FKey> CurrentKeyValue = CurrentKey.Get();
 	if (CurrentKeyValue.IsSet())
 	{
+		const FKey& Key = CurrentKeyValue.GetValue();
+		if (Key.IsValid() && (Key.IsDeprecated() || !Key.IsBindableToActions()))
+		{
+			return FEditorStyle::GetBrush("Icons.Warning");
+		}
 		return GetIconFromKey(CurrentKeyValue.GetValue());
 	}
 	return nullptr;
 }
 
+FText SKeySelector::GetKeyTooltip() const
+{
+	TOptional<FKey> CurrentKeyValue = CurrentKey.Get();
+	if (CurrentKeyValue.IsSet())
+	{
+		const FKey& Key = CurrentKeyValue.GetValue();
+		if (Key.IsValid())
+		{
+			if (Key.IsDeprecated())
+			{
+				return LOCTEXT("KeySelectorDeprecated", "The selected key has been deprecated.");
+			}
+			else if (!Key.IsBindableToActions())
+			{
+				return LOCTEXT("KeySelectorNotBindable", "The selected key can't be bound to actions.");
+			}
+		}
+	}
+	return LOCTEXT("KeySelector", "Select the key value.");
+}
+
 //=======================================================================
 // Key TreeView Support
 TSharedRef<ITableRow> SKeySelector::GenerateKeyTreeRow(FKeyTreeItem InItem, const TSharedRef<STableViewBase>& OwnerTree)
diff --git a/Engine/Source/Editor/UnrealEd/Public/SKeySelector.h b/Engine/Source/Editor/UnrealEd/Public/SKeySelector.h
--- a/Engine/Source/Editor/UnrealEd/Public/SKeySelector.h
+++ b/Engine/Source/Editor/UnrealEd/Public/SKeySelector.h
@@ -55,6 +55,8 @@ protected:
 	const FSlateBrush* GetKeyIconImage() const;
 	/** Gets a succinct description for the key being manipulated */
 	FText GetKeyDescription() const;
+	/** Gets a tooltip for the selected key */
+	FText GetKeyTooltip() const;
 
 	/** Treeview support functions */
 	virtual TSharedRef<ITableRow> GenerateKeyTreeRow(FKeyTreeItem InItem, const TSharedRef<STableViewBase>& OwnerTree);
diff --git a/Engine/Source/Runtime/Engine/Private/UserInterface/InputSettings.cpp b/Engine/Source/Runtime/Engine/Private/UserInterface/InputSettings.cpp
--- a/Engine/Source/Runtime/Engine/Private/UserInterface/InputSettings.cpp
+++ b/Engine/Source/Runtime/Engine/Private/UserInterface/InputSettings.cpp
@@ -73,6 +73,22 @@ void UInputSettings::PostInitProperties()
 		}
 	}
 #endif
+
+	for (const FInputActionKeyMapping& KeyMapping : ActionMappings)
+	{
+		if (KeyMapping.Key.IsDeprecated())
+		{
+			UE_LOG(LogInput, Warning, TEXT("Action %s uses deprecated key %s."), *KeyMapping.ActionName.ToString(), *KeyMapping.Key.ToString());
+		}
+	}
+
+	for (const FInputAxisKeyMapping& KeyMapping : AxisMappings)
+	{
+		if (KeyMapping.Key.IsDeprecated())
+		{
+			UE_LOG(LogInput, Warning, TEXT("Axis %s uses deprecated key %s."), *KeyMapping.AxisName.ToString(), *KeyMapping.Key.ToString());
+		}
+	}
 }
 
 void UInputSettings::PopulateAxisConfigs()
diff --git a/Templates/TP_CollaborativeBP/Config/DefaultInput.ini b/Templates/TP_CollaborativeBP/Config/DefaultInput.ini
--- a/Templates/TP_CollaborativeBP/Config/DefaultInput.ini
+++ b/Templates/TP_CollaborativeBP/Config/DefaultInput.ini
@@ -1,5 +1,11 @@
 
 [/Script/Engine.InputSettings]
++AxisConfig=(AxisKeyName="OculusTouch_Left_Thumbstick_Y",AxisProperties=(DeadZone=0.000000,Sensitivity=1.000000,Exponent=1.000000,bInvert=False))
++AxisConfig=(AxisKeyName="OculusTouch_Right_Grip_Axis",AxisProperties=(DeadZone=0.000000,Sensitivity=1.000000,Exponent=1.000000,bInvert=False))
++AxisConfig=(AxisKeyName="OculusTouch_Right_Trigger_Axis",AxisProperties=(DeadZone=0.000000,Sensitivity=1.000000,Exponent=1.000000,bInvert=False))
++AxisConfig=(AxisKeyName="OculusTouch_Right_Thumbstick_X",AxisProperties=(DeadZone=0.000000,Sensitivity=1.000000,Exponent=1.000000,bInvert=False))
++AxisConfig=(AxisKeyName="OculusTouch_Right_Thumbstick_Y",AxisProperties=(DeadZone=0.000000,Sensitivity=1.000000,Exponent=1.000000,bInvert=False))
++AxisConfig=(AxisKeyName="ValveIndex_Left_Grip_Axis",AxisProperties=(DeadZone=0.000000,Sensitivity=1.000000,Exponent=1.000000,bInvert=False))
 +AxisConfig=(AxisKeyName="ValveIndex_Left_Grip_Force",AxisProperties=(DeadZone=0.000000,Sensitivity=1.000000,Exponent=1.000000,bInvert=False))
 +AxisConfig=(AxisKeyName="ValveIndex_Left_Trigger_Axis",AxisProperties=(DeadZone=0.000000,Sensitivity=1.000000,Exponent=1.000000,bInvert=False))
 +AxisConfig=(AxisKeyName="ValveIndex_Left_Thumbstick_X",AxisProperties=(DeadZone=0.000000,Sensitivity=1.000000,Exponent=1.000000,bInvert=False))
@@ -104,19 +110,12 @@
 +AxisConfig=(AxisKeyName="OculusTouch_Left_Grip_Axis",AxisProperties=(DeadZone=0.000000,Sensitivity=1.000000,Exponent=1.000000,bInvert=False))
 +AxisConfig=(AxisKeyName="OculusTouch_Left_Trigger_Axis",AxisProperties=(DeadZone=0.000000,Sensitivity=1.000000,Exponent=1.000000,bInvert=False))
 +AxisConfig=(AxisKeyName="OculusTouch_Left_Thumbstick_X",AxisProperties=(DeadZone=0.000000,Sensitivity=1.000000,Exponent=1.000000,bInvert=False))
-+AxisConfig=(AxisKeyName="OculusTouch_Left_Thumbstick_Y",AxisProperties=(DeadZone=0.000000,Sensitivity=1.000000,Exponent=1.000000,bInvert=False))
-+AxisConfig=(AxisKeyName="OculusTouch_Right_Grip_Axis",AxisProperties=(DeadZone=0.000000,Sensitivity=1.000000,Exponent=1.000000,bInvert=False))
-+AxisConfig=(AxisKeyName="OculusTouch_Right_Trigger_Axis",AxisProperties=(DeadZone=0.000000,Sensitivity=1.000000,Exponent=1.000000,bInvert=False))
-+AxisConfig=(AxisKeyName="OculusTouch_Right_Thumbstick_X",AxisProperties=(DeadZone=0.000000,Sensitivity=1.000000,Exponent=1.000000,bInvert=False))
-+AxisConfig=(AxisKeyName="OculusTouch_Right_Thumbstick_Y",AxisProperties=(DeadZone=0.000000,Sensitivity=1.000000,Exponent=1.000000,bInvert=False))
-+AxisConfig=(AxisKeyName="ValveIndex_Left_Grip_Axis",AxisProperties=(DeadZone=0.000000,Sensitivity=1.000000,Exponent=1.000000,bInvert=False))
 bAltEnterTogglesFullscreen=True
 bF11TogglesFullscreen=True
 bUseMouseForTouch=False
 bEnableMouseSmoothing=True
 bEnableFOVScaling=True
 bCaptureMouseOnLaunch=True
-bDefaultViewportMouseLock=False
 bAlwaysShowTouchInterface=False
 bShowConsoleOnFourFingerTap=True
 bEnableGestureRecognizer=True
@@ -133,7 +132,7 @@ DoubleClickTime=0.200000
 +ActionMappings=(ActionName="WalkMode",bShift=False,bCtrl=True,bAlt=False,bCmd=False,Key=Two)
 +ActionMappings=(ActionName="OrbitMode",bShift=False,bCtrl=True,bAlt=False,bCmd=False,Key=Three)
 +ActionMappings=(ActionName="VRMode",bShift=False,bCtrl=True,bAlt=False,bCmd=False,Key=Four)
-+ActionMappings=(ActionName="ResetLocationL",bShift=False,bCtrl=False,bAlt=False,bCmd=False,Key=MotionController_Left_Grip1)
++ActionMappings=(ActionName="ResetLocationL",bShift=False,bCtrl=False,bAlt=False,bCmd=False,Key=OculusTouch_Left_Thumbstick_Click)
 +ActionMappings=(ActionName="Bookmark_1",bShift=False,bCtrl=False,bAlt=False,bCmd=False,Key=One)
 +ActionMappings=(ActionName="Bookmark_2",bShift=False,bCtrl=False,bAlt=False,bCmd=False,Key=Two)
 +ActionMappings=(ActionName="Bookmark_3",bShift=False,bCtrl=False,bAlt=False,bCmd=False,Key=Three)
@@ -179,6 +178,8 @@ DoubleClickTime=0.200000
 +ActionMappings=(ActionName="ShowContextualMenu_L",bShift=False,bCtrl=False,bAlt=False,bCmd=False,Key=ValveIndex_Left_A_Click)
 +ActionMappings=(ActionName="VR_Teleport_L",bShift=False,bCtrl=False,bAlt=False,bCmd=False,Key=ValveIndex_Left_B_Click)
 +ActionMappings=(ActionName="VR_Teleport_R",bShift=False,bCtrl=False,bAlt=False,bCmd=False,Key=ValveIndex_Right_B_Click)
++ActionMappings=(ActionName="ResetLocationL",bShift=False,bCtrl=False,bAlt=False,bCmd=False,Key=Vive_Left_Grip_Click)
++ActionMappings=(ActionName="ResetLocationL",bShift=False,bCtrl=False,bAlt=False,bCmd=False,Key=ValveIndex_Left_Thumbstick_Click)
 +AxisMappings=(AxisName="WalkForward",Scale=1.000000,Key=W)
 +AxisMappings=(AxisName="WalkRight",Scale=1.000000,Key=D)
 +AxisMappings=(AxisName="WalkForward",Scale=-1.000000,Key=S)
@@ -209,7 +210,6 @@ DoubleClickTime=0.200000
 +AxisMappings=(AxisName="ThumbStick_L_Y",Scale=-1.000000,Key=ValveIndex_Left_Thumbstick_Y)
 +AxisMappings=(AxisName="ThumbStick_R_Y",Scale=-1.000000,Key=ValveIndex_Right_Thumbstick_Y)
 DefaultTouchInterface=/Game/CollaborativeViewer/Blueprints/Pawn/Touch/TouchInterface.TouchInterface
-ConsoleKey=None
 +ConsoleKeys=Tilde
 
 
diff --git a/Templates/TP_FirstPerson/Config/DefaultInput.ini b/Templates/TP_FirstPerson/Config/DefaultInput.ini
--- a/Templates/TP_FirstPerson/Config/DefaultInput.ini
+++ b/Templates/TP_FirstPerson/Config/DefaultInput.ini
@@ -1,11 +1,5 @@
 
 [/Script/Engine.InputSettings]
--AxisConfig=(AxisKeyName="Gamepad_LeftX",AxisProperties=(DeadZone=0.25,Exponent=1.f,Sensitivity=1.f))
--AxisConfig=(AxisKeyName="Gamepad_LeftY",AxisProperties=(DeadZone=0.25,Exponent=1.f,Sensitivity=1.f))
--AxisConfig=(AxisKeyName="Gamepad_RightX",AxisProperties=(DeadZone=0.25,Exponent=1.f,Sensitivity=1.f))
--AxisConfig=(AxisKeyName="Gamepad_RightY",AxisProperties=(DeadZone=0.25,Exponent=1.f,Sensitivity=1.f))
--AxisConfig=(AxisKeyName="MouseX",AxisProperties=(DeadZone=0.f,Exponent=1.f,Sensitivity=0.07f))
--AxisConfig=(AxisKeyName="MouseY",AxisProperties=(DeadZone=0.f,Exponent=1.f,Sensitivity=0.07f))
 +AxisConfig=(AxisKeyName="Gamepad_LeftX",AxisProperties=(DeadZone=0.250000,Sensitivity=1.000000,Exponent=1.000000,bInvert=False))
 +AxisConfig=(AxisKeyName="Gamepad_LeftY",AxisProperties=(DeadZone=0.250000,Sensitivity=1.000000,Exponent=1.000000,bInvert=False))
 +AxisConfig=(AxisKeyName="Gamepad_RightX",AxisProperties=(DeadZone=0.250000,Sensitivity=1.000000,Exponent=1.000000,bInvert=False))
@@ -53,13 +47,75 @@
 +AxisConfig=(AxisKeyName="SteamVR_Knuckles_Right_MiddleGrip",AxisProperties=(DeadZone=0.000000,Sensitivity=1.000000,Exponent=1.000000,bInvert=False))
 +AxisConfig=(AxisKeyName="SteamVR_Knuckles_Right_RingGrip",AxisProperties=(DeadZone=0.000000,Sensitivity=1.000000,Exponent=1.000000,bInvert=False))
 +AxisConfig=(AxisKeyName="SteamVR_Knuckles_Right_PinkyGrip",AxisProperties=(DeadZone=0.000000,Sensitivity=1.000000,Exponent=1.000000,bInvert=False))
++AxisConfig=(AxisKeyName="Daydream_Left_Trackpad_X",AxisProperties=(DeadZone=0.000000,Sensitivity=1.000000,Exponent=1.000000,bInvert=False))
++AxisConfig=(AxisKeyName="Daydream_Left_Trackpad_Y",AxisProperties=(DeadZone=0.000000,Sensitivity=1.000000,Exponent=1.000000,bInvert=False))
++AxisConfig=(AxisKeyName="Daydream_Right_Trackpad_X",AxisProperties=(DeadZone=0.000000,Sensitivity=1.000000,Exponent=1.000000,bInvert=False))
++AxisConfig=(AxisKeyName="Daydream_Right_Trackpad_Y",AxisProperties=(DeadZone=0.000000,Sensitivity=1.000000,Exponent=1.000000,bInvert=False))
++AxisConfig=(AxisKeyName="Vive_Left_Trigger_Axis",AxisProperties=(DeadZone=0.000000,Sensitivity=1.000000,Exponent=1.000000,bInvert=False))
++AxisConfig=(AxisKeyName="Vive_Left_Trackpad_X",AxisProperties=(DeadZone=0.000000,Sensitivity=1.000000,Exponent=1.000000,bInvert=False))
++AxisConfig=(AxisKeyName="Vive_Left_Trackpad_Y",AxisProperties=(DeadZone=0.000000,Sensitivity=1.000000,Exponent=1.000000,bInvert=False))
++AxisConfig=(AxisKeyName="Vive_Right_Trigger_Axis",AxisProperties=(DeadZone=0.000000,Sensitivity=1.000000,Exponent=1.000000,bInvert=False))
++AxisConfig=(AxisKeyName="Vive_Right_Trackpad_X",AxisProperties=(DeadZone=0.000000,Sensitivity=1.000000,Exponent=1.000000,bInvert=False))
++AxisConfig=(AxisKeyName="Vive_Right_Trackpad_Y",AxisProperties=(DeadZone=0.000000,Sensitivity=1.000000,Exponent=1.000000,bInvert=False))
++AxisConfig=(AxisKeyName="MixedReality_Left_Trigger_Axis",AxisProperties=(DeadZone=0.000000,Sensitivity=1.000000,Exponent=1.000000,bInvert=False))
++AxisConfig=(AxisKeyName="MixedReality_Left_Thumbstick_X",AxisProperties=(DeadZone=0.000000,Sensitivity=1.000000,Exponent=1.000000,bInvert=False))
++AxisConfig=(AxisKeyName="MixedReality_Left_Thumbstick_Y",AxisProperties=(DeadZone=0.000000,Sensitivity=1.000000,Exponent=1.000000,bInvert=False))
++AxisConfig=(AxisKeyName="MixedReality_Left_Trackpad_X",AxisProperties=(DeadZone=0.000000,Sensitivity=1.000000,Exponent=1.000000,bInvert=False))
++AxisConfig=(AxisKeyName="MixedReality_Left_Trackpad_Y",AxisProperties=(DeadZone=0.000000,Sensitivity=1.000000,Exponent=1.000000,bInvert=False))
++AxisConfig=(AxisKeyName="MixedReality_Right_Trigger_Axis",AxisProperties=(DeadZone=0.000000,Sensitivity=1.000000,Exponent=1.000000,bInvert=False))
++AxisConfig=(AxisKeyName="MixedReality_Right_Thumbstick_X",AxisProperties=(DeadZone=0.000000,Sensitivity=1.000000,Exponent=1.000000,bInvert=False))
++AxisConfig=(AxisKeyName="MixedReality_Right_Thumbstick_Y",AxisProperties=(DeadZone=0.000000,Sensitivity=1.000000,Exponent=1.000000,bInvert=False))
++AxisConfig=(AxisKeyName="MixedReality_Right_Trackpad_X",AxisProperties=(DeadZone=0.000000,Sensitivity=1.000000,Exponent=1.000000,bInvert=False))
++AxisConfig=(AxisKeyName="MixedReality_Right_Trackpad_Y",AxisProperties=(DeadZone=0.000000,Sensitivity=1.000000,Exponent=1.000000,bInvert=False))
++AxisConfig=(AxisKeyName="OculusGo_Left_Trackpad_X",AxisProperties=(DeadZone=0.000000,Sensitivity=1.000000,Exponent=1.000000,bInvert=False))
++AxisConfig=(AxisKeyName="OculusGo_Left_Trackpad_Y",AxisProperties=(DeadZone=0.000000,Sensitivity=1.000000,Exponent=1.000000,bInvert=False))
++AxisConfig=(AxisKeyName="OculusGo_Right_Trackpad_X",AxisProperties=(DeadZone=0.000000,Sensitivity=1.000000,Exponent=1.000000,bInvert=False))
++AxisConfig=(AxisKeyName="OculusGo_Right_Trackpad_Y",AxisProperties=(DeadZone=0.000000,Sensitivity=1.000000,Exponent=1.000000,bInvert=False))
++AxisConfig=(AxisKeyName="OculusTouch_Left_Grip_Axis",AxisProperties=(DeadZone=0.000000,Sensitivity=1.000000,Exponent=1.000000,bInvert=False))
++AxisConfig=(AxisKeyName="OculusTouch_Left_Trigger_Axis",AxisProperties=(DeadZone=0.000000,Sensitivity=1.000000,Exponent=1.000000,bInvert=False))
++AxisConfig=(AxisKeyName="OculusTouch_Left_Thumbstick_X",AxisProperties=(DeadZone=0.000000,Sensitivity=1.000000,Exponent=1.000000,bInvert=False))
++AxisConfig=(AxisKeyName="OculusTouch_Left_Thumbstick_Y",AxisProperties=(DeadZone=0.000000,Sensitivity=1.000000,Exponent=1.000000,bInvert=False))
++AxisConfig=(AxisKeyName="OculusTouch_Right_Grip_Axis",AxisProperties=(DeadZone=0.000000,Sensitivity=1.000000,Exponent=1.000000,bInvert=False))
++AxisConfig=(AxisKeyName="OculusTouch_Right_Trigger_Axis",AxisProperties=(DeadZone=0.000000,Sensitivity=1.000000,Exponent=1.000000,bInvert=False))
++AxisConfig=(AxisKeyName="OculusTouch_Right_Thumbstick_X",AxisProperties=(DeadZone=0.000000,Sensitivity=1.000000,Exponent=1.000000,bInvert=False))
++AxisConfig=(AxisKeyName="OculusTouch_Right_Thumbstick_Y",AxisProperties=(DeadZone=0.000000,Sensitivity=1.000000,Exponent=1.000000,bInvert=False))
++AxisConfig=(AxisKeyName="ValveIndex_Left_Grip_Axis",AxisProperties=(DeadZone=0.000000,Sensitivity=1.000000,Exponent=1.000000,bInvert=False))
++AxisConfig=(AxisKeyName="ValveIndex_Left_Grip_Force",AxisProperties=(DeadZone=0.000000,Sensitivity=1.000000,Exponent=1.000000,bInvert=False))
++AxisConfig=(AxisKeyName="ValveIndex_Left_Trigger_Axis",AxisProperties=(DeadZone=0.000000,Sensitivity=1.000000,Exponent=1.000000,bInvert=False))
++AxisConfig=(AxisKeyName="ValveIndex_Left_Thumbstick_X",AxisProperties=(DeadZone=0.000000,Sensitivity=1.000000,Exponent=1.000000,bInvert=False))
++AxisConfig=(AxisKeyName="ValveIndex_Left_Thumbstick_Y",AxisProperties=(DeadZone=0.000000,Sensitivity=1.000000,Exponent=1.000000,bInvert=False))
++AxisConfig=(AxisKeyName="ValveIndex_Left_Trackpad_X",AxisProperties=(DeadZone=0.000000,Sensitivity=1.000000,Exponent=1.000000,bInvert=False))
++AxisConfig=(AxisKeyName="ValveIndex_Left_Trackpad_Y",AxisProperties=(DeadZone=0.000000,Sensitivity=1.000000,Exponent=1.000000,bInvert=False))
++AxisConfig=(AxisKeyName="ValveIndex_Left_Trackpad_Force",AxisProperties=(DeadZone=0.000000,Sensitivity=1.000000,Exponent=1.000000,bInvert=False))
++AxisConfig=(AxisKeyName="ValveIndex_Left_Trackpad_Touch",AxisProperties=(DeadZone=0.000000,Sensitivity=1.000000,Exponent=1.000000,bInvert=False))
++AxisConfig=(AxisKeyName="ValveIndex_Right_Grip_Axis",AxisProperties=(DeadZone=0.000000,Sensitivity=1.000000,Exponent=1.000000,bInvert=False))
++AxisConfig=(AxisKeyName="ValveIndex_Right_Grip_Force",AxisProperties=(DeadZone=0.000000,Sensitivity=1.000000,Exponent=1.000000,bInvert=False))
++AxisConfig=(AxisKeyName="ValveIndex_Right_Trigger_Axis",AxisProperties=(DeadZone=0.000000,Sensitivity=1.000000,Exponent=1.000000,bInvert=False))
++AxisConfig=(AxisKeyName="ValveIndex_Right_Thumbstick_X",AxisProperties=(DeadZone=0.000000,Sensitivity=1.000000,Exponent=1.000000,bInvert=False))
++AxisConfig=(AxisKeyName="ValveIndex_Right_Thumbstick_Y",AxisProperties=(DeadZone=0.000000,Sensitivity=1.000000,Exponent=1.000000,bInvert=False))
++AxisConfig=(AxisKeyName="ValveIndex_Right_Trackpad_X",AxisProperties=(DeadZone=0.000000,Sensitivity=1.000000,Exponent=1.000000,bInvert=False))
++AxisConfig=(AxisKeyName="ValveIndex_Right_Trackpad_Y",AxisProperties=(DeadZone=0.000000,Sensitivity=1.000000,Exponent=1.000000,bInvert=False))
++AxisConfig=(AxisKeyName="ValveIndex_Right_Trackpad_Force",AxisProperties=(DeadZone=0.000000,Sensitivity=1.000000,Exponent=1.000000,bInvert=False))
++AxisConfig=(AxisKeyName="MagicLeap_Left_Trigger_Axis",AxisProperties=(DeadZone=0.000000,Sensitivity=1.000000,Exponent=1.000000,bInvert=False))
++AxisConfig=(AxisKeyName="MagicLeap_Left_Trackpad_X",AxisProperties=(DeadZone=0.000000,Sensitivity=1.000000,Exponent=1.000000,bInvert=False))
++AxisConfig=(AxisKeyName="MagicLeap_Left_Trackpad_Y",AxisProperties=(DeadZone=0.000000,Sensitivity=1.000000,Exponent=1.000000,bInvert=False))
++AxisConfig=(AxisKeyName="MagicLeap_Left_Trackpad_Force",AxisProperties=(DeadZone=0.000000,Sensitivity=1.000000,Exponent=1.000000,bInvert=False))
++AxisConfig=(AxisKeyName="MagicLeap_Left_Touch1_X",AxisProperties=(DeadZone=0.000000,Sensitivity=1.000000,Exponent=1.000000,bInvert=False))
++AxisConfig=(AxisKeyName="MagicLeap_Left_Touch1_Y",AxisProperties=(DeadZone=0.000000,Sensitivity=1.000000,Exponent=1.000000,bInvert=False))
++AxisConfig=(AxisKeyName="MagicLeap_Left_Touch1_Force",AxisProperties=(DeadZone=0.000000,Sensitivity=1.000000,Exponent=1.000000,bInvert=False))
++AxisConfig=(AxisKeyName="MagicLeap_Right_Trigger_Axis",AxisProperties=(DeadZone=0.000000,Sensitivity=1.000000,Exponent=1.000000,bInvert=False))
++AxisConfig=(AxisKeyName="MagicLeap_Right_Trackpad_X",AxisProperties=(DeadZone=0.000000,Sensitivity=1.000000,Exponent=1.000000,bInvert=False))
++AxisConfig=(AxisKeyName="MagicLeap_Right_Trackpad_Y",AxisProperties=(DeadZone=0.000000,Sensitivity=1.000000,Exponent=1.000000,bInvert=False))
++AxisConfig=(AxisKeyName="MagicLeap_Right_Trackpad_Force",AxisProperties=(DeadZone=0.000000,Sensitivity=1.000000,Exponent=1.000000,bInvert=False))
++AxisConfig=(AxisKeyName="MagicLeap_Right_Touch1_X",AxisProperties=(DeadZone=0.000000,Sensitivity=1.000000,Exponent=1.000000,bInvert=False))
++AxisConfig=(AxisKeyName="MagicLeap_Right_Touch1_Y",AxisProperties=(DeadZone=0.000000,Sensitivity=1.000000,Exponent=1.000000,bInvert=False))
++AxisConfig=(AxisKeyName="MagicLeap_Right_Touch1_Force",AxisProperties=(DeadZone=0.000000,Sensitivity=1.000000,Exponent=1.000000,bInvert=False))
 bAltEnterTogglesFullscreen=True
 bF11TogglesFullscreen=True
 bUseMouseForTouch=False
 bEnableMouseSmoothing=True
 bEnableFOVScaling=True
 bCaptureMouseOnLaunch=True
-bDefaultViewportMouseLock=False
 bAlwaysShowTouchInterface=False
 bShowConsoleOnFourFingerTap=True
 bEnableGestureRecognizer=False
@@ -70,32 +126,58 @@ FOVScale=0.011110
 DoubleClickTime=0.200000
 +ActionMappings=(ActionName="Jump",bShift=False,bCtrl=False,bAlt=False,bCmd=False,Key=SpaceBar)
 +ActionMappings=(ActionName="Jump",bShift=False,bCtrl=False,bAlt=False,bCmd=False,Key=Gamepad_FaceButton_Bottom)
-+ActionMappings=(ActionName="Jump",bShift=False,bCtrl=False,bAlt=False,bCmd=False,Key=MotionController_Left_Trigger)
++ActionMappings=(ActionName="Jump",bShift=False,bCtrl=False,bAlt=False,bCmd=False,Key=Daydream_Left_Select_Click)
 +ActionMappings=(ActionName="Fire",bShift=False,bCtrl=False,bAlt=False,bCmd=False,Key=LeftMouseButton)
 +ActionMappings=(ActionName="Fire",bShift=False,bCtrl=False,bAlt=False,bCmd=False,Key=Gamepad_RightTrigger)
-+ActionMappings=(ActionName="Fire",bShift=False,bCtrl=False,bAlt=False,bCmd=False,Key=MotionController_Right_Trigger)
++ActionMappings=(ActionName="Fire",bShift=False,bCtrl=False,bAlt=False,bCmd=False,Key=Daydream_Left_Trackpad_Click)
 +ActionMappings=(ActionName="ResetVR",bShift=False,bCtrl=False,bAlt=False,bCmd=False,Key=R)
-+ActionMappings=(ActionName="ResetVR",bShift=False,bCtrl=False,bAlt=False,bCmd=False,Key=MotionController_Left_Grip1)
-+ActionMappings=(ActionName="Fire",bShift=False,bCtrl=False,bAlt=False,bCmd=False,Key=OculusTouchpad_Touchpad)
++ActionMappings=(ActionName="ResetVR",bShift=False,bCtrl=False,bAlt=False,bCmd=False,Key=Vive_Left_Grip_Click)
++ActionMappings=(ActionName="Fire",bShift=False,bCtrl=False,bAlt=False,bCmd=False,Key=Vive_Right_Trigger_Click)
++ActionMappings=(ActionName="Jump",bShift=False,bCtrl=False,bAlt=False,bCmd=False,Key=Vive_Left_Trigger_Click)
++ActionMappings=(ActionName="Jump",bShift=False,bCtrl=False,bAlt=False,bCmd=False,Key=MixedReality_Left_Trigger_Click)
++ActionMappings=(ActionName="Jump",bShift=False,bCtrl=False,bAlt=False,bCmd=False,Key=OculusGo_Left_Trigger_Click)
++ActionMappings=(ActionName="Jump",bShift=False,bCtrl=False,bAlt=False,bCmd=False,Key=OculusTouch_Left_Trigger_Click)
++ActionMappings=(ActionName="Jump",bShift=False,bCtrl=False,bAlt=False,bCmd=False,Key=ValveIndex_Left_Trigger_Click)
++ActionMappings=(ActionName="Jump",bShift=False,bCtrl=False,bAlt=False,bCmd=False,Key=MagicLeap_Left_Trigger)
++ActionMappings=(ActionName="Fire",bShift=False,bCtrl=False,bAlt=False,bCmd=False,Key=MixedReality_Right_Trigger_Click)
++ActionMappings=(ActionName="Fire",bShift=False,bCtrl=False,bAlt=False,bCmd=False,Key=OculusGo_Left_Trackpad_Click)
++ActionMappings=(ActionName="Fire",bShift=False,bCtrl=False,bAlt=False,bCmd=False,Key=OculusTouch_Right_Trigger_Click)
++ActionMappings=(ActionName="Fire",bShift=False,bCtrl=False,bAlt=False,bCmd=False,Key=ValveIndex_Right_Trigger_Click)
++ActionMappings=(ActionName="Fire",bShift=False,bCtrl=False,bAlt=False,bCmd=False,Key=MagicLeap_Left_Trackpad_Touch)
++ActionMappings=(ActionName="ResetVR",bShift=False,bCtrl=False,bAlt=False,bCmd=False,Key=MixedReality_Left_Thumbstick_Click)
++ActionMappings=(ActionName="ResetVR",bShift=False,bCtrl=False,bAlt=False,bCmd=False,Key=OculusGo_Left_Back_Click)
++ActionMappings=(ActionName="ResetVR",bShift=False,bCtrl=False,bAlt=False,bCmd=False,Key=OculusTouch_Left_Thumbstick_Click)
++ActionMappings=(ActionName="ResetVR",bShift=False,bCtrl=False,bAlt=False,bCmd=False,Key=ValveIndex_Left_Thumbstick_Click)
++ActionMappings=(ActionName="ResetVR",bShift=False,bCtrl=False,bAlt=False,bCmd=False,Key=MagicLeap_Left_Bumper)
 +AxisMappings=(AxisName="MoveForward",Scale=1.000000,Key=W)
 +AxisMappings=(AxisName="MoveForward",Scale=-1.000000,Key=S)
 +AxisMappings=(AxisName="MoveForward",Scale=1.000000,Key=Up)
 +AxisMappings=(AxisName="MoveForward",Scale=-1.000000,Key=Down)
 +AxisMappings=(AxisName="MoveForward",Scale=1.000000,Key=Gamepad_LeftY)
-+AxisMappings=(AxisName="MoveForward",Scale=-1.000000,Key=MotionController_Left_Thumbstick_Y)
 +AxisMappings=(AxisName="MoveRight",Scale=-1.000000,Key=A)
 +AxisMappings=(AxisName="MoveRight",Scale=1.000000,Key=D)
 +AxisMappings=(AxisName="MoveRight",Scale=1.000000,Key=Gamepad_LeftX)
-+AxisMappings=(AxisName="MoveRight",Scale=1.000000,Key=MotionController_Left_Thumbstick_X)
++AxisMappings=(AxisName="MoveRight",Scale=1.000000,Key=Daydream_Left_Trackpad_X)
 +AxisMappings=(AxisName="TurnRate",Scale=1.000000,Key=Gamepad_RightX)
 +AxisMappings=(AxisName="TurnRate",Scale=-1.000000,Key=Left)
 +AxisMappings=(AxisName="TurnRate",Scale=1.000000,Key=Right)
 +AxisMappings=(AxisName="Turn",Scale=1.000000,Key=MouseX)
 +AxisMappings=(AxisName="LookUpRate",Scale=1.000000,Key=Gamepad_RightY)
 +AxisMappings=(AxisName="LookUp",Scale=-1.000000,Key=MouseY)
++AxisMappings=(AxisName="MoveForward",Scale=1.000000,Key=Daydream_Left_Trackpad_Y)
++AxisMappings=(AxisName="MoveForward",Scale=1.000000,Key=Vive_Left_Trackpad_Y)
++AxisMappings=(AxisName="MoveForward",Scale=1.000000,Key=MixedReality_Left_Thumbstick_Y)
++AxisMappings=(AxisName="MoveForward",Scale=1.000000,Key=OculusGo_Left_Trackpad_Y)
++AxisMappings=(AxisName="MoveForward",Scale=1.000000,Key=OculusTouch_Left_Thumbstick_Y)
++AxisMappings=(AxisName="MoveForward",Scale=1.000000,Key=ValveIndex_Left_Thumbstick_Y)
++AxisMappings=(AxisName="MoveRight",Scale=1.000000,Key=Vive_Left_Trackpad_X)
++AxisMappings=(AxisName="MoveRight",Scale=1.000000,Key=MixedReality_Left_Thumbstick_X)
++AxisMappings=(AxisName="MoveRight",Scale=1.000000,Key=OculusGo_Left_Trackpad_X)
++AxisMappings=(AxisName="MoveRight",Scale=1.000000,Key=OculusTouch_Left_Thumbstick_Y)
++AxisMappings=(AxisName="MoveRight",Scale=1.000000,Key=ValveIndex_Left_Thumbstick_Y)
++AxisMappings=(AxisName="MoveRight",Scale=1.000000,Key=MagicLeap_Left_Trackpad_X)
++AxisMappings=(AxisName="MoveForward",Scale=1.000000,Key=MagicLeap_Left_Trackpad_Y)
 DefaultTouchInterface=/Engine/MobileResources/HUD/DefaultVirtualJoysticks.DefaultVirtualJoysticks
-ConsoleKey=None
--ConsoleKeys=Tilde
 +ConsoleKeys=Tilde
 
 
diff --git a/Templates/TP_FirstPersonBP/Config/DefaultInput.ini b/Templates/TP_FirstPersonBP/Config/DefaultInput.ini
--- a/Templates/TP_FirstPersonBP/Config/DefaultInput.ini
+++ b/Templates/TP_FirstPersonBP/Config/DefaultInput.ini
@@ -1,17 +1,13 @@
 
 [/Script/Engine.InputSettings]
--AxisConfig=(AxisKeyName="Gamepad_LeftX",AxisProperties=(DeadZone=0.25,Exponent=1.f,Sensitivity=1.f))
--AxisConfig=(AxisKeyName="Gamepad_LeftY",AxisProperties=(DeadZone=0.25,Exponent=1.f,Sensitivity=1.f))
--AxisConfig=(AxisKeyName="Gamepad_RightX",AxisProperties=(DeadZone=0.25,Exponent=1.f,Sensitivity=1.f))
--AxisConfig=(AxisKeyName="Gamepad_RightY",AxisProperties=(DeadZone=0.25,Exponent=1.f,Sensitivity=1.f))
--AxisConfig=(AxisKeyName="MouseX",AxisProperties=(DeadZone=0.f,Exponent=1.f,Sensitivity=0.07f))
--AxisConfig=(AxisKeyName="MouseY",AxisProperties=(DeadZone=0.f,Exponent=1.f,Sensitivity=0.07f))
 +AxisConfig=(AxisKeyName="Gamepad_LeftX",AxisProperties=(DeadZone=0.250000,Sensitivity=1.000000,Exponent=1.000000,bInvert=False))
 +AxisConfig=(AxisKeyName="Gamepad_LeftY",AxisProperties=(DeadZone=0.250000,Sensitivity=1.000000,Exponent=1.000000,bInvert=False))
 +AxisConfig=(AxisKeyName="Gamepad_RightX",AxisProperties=(DeadZone=0.250000,Sensitivity=1.000000,Exponent=1.000000,bInvert=False))
 +AxisConfig=(AxisKeyName="Gamepad_RightY",AxisProperties=(DeadZone=0.250000,Sensitivity=1.000000,Exponent=1.000000,bInvert=False))
 +AxisConfig=(AxisKeyName="MouseX",AxisProperties=(DeadZone=0.000000,Sensitivity=0.070000,Exponent=1.000000,bInvert=False))
 +AxisConfig=(AxisKeyName="MouseY",AxisProperties=(DeadZone=0.000000,Sensitivity=0.070000,Exponent=1.000000,bInvert=False))
++AxisConfig=(AxisKeyName="MotionController_Left_Thumbstick_Z",AxisProperties=(DeadZone=0.000000,Sensitivity=1.000000,Exponent=1.000000,bInvert=False))
++AxisConfig=(AxisKeyName="MotionController_Right_Thumbstick_Z",AxisProperties=(DeadZone=0.000000,Sensitivity=1.000000,Exponent=1.000000,bInvert=False))
 +AxisConfig=(AxisKeyName="MouseWheelAxis",AxisProperties=(DeadZone=0.000000,Sensitivity=1.000000,Exponent=1.000000,bInvert=False))
 +AxisConfig=(AxisKeyName="Gamepad_LeftTriggerAxis",AxisProperties=(DeadZone=0.000000,Sensitivity=1.000000,Exponent=1.000000,bInvert=False))
 +AxisConfig=(AxisKeyName="Gamepad_RightTriggerAxis",AxisProperties=(DeadZone=0.000000,Sensitivity=1.000000,Exponent=1.000000,bInvert=False))
@@ -27,8 +23,6 @@
 +AxisConfig=(AxisKeyName="MotionController_Right_Grip2Axis",AxisProperties=(DeadZone=0.000000,Sensitivity=1.000000,Exponent=1.000000,bInvert=False))
 +AxisConfig=(AxisKeyName="Gamepad_Special_Left_X",AxisProperties=(DeadZone=0.000000,Sensitivity=1.000000,Exponent=1.000000,bInvert=False))
 +AxisConfig=(AxisKeyName="Gamepad_Special_Left_Y",AxisProperties=(DeadZone=0.000000,Sensitivity=1.000000,Exponent=1.000000,bInvert=False))
-+AxisConfig=(AxisKeyName="MotionController_Left_Thumbstick_Z",AxisProperties=(DeadZone=0.000000,Sensitivity=1.000000,Exponent=1.000000,bInvert=False))
-+AxisConfig=(AxisKeyName="MotionController_Right_Thumbstick_Z",AxisProperties=(DeadZone=0.000000,Sensitivity=1.000000,Exponent=1.000000,bInvert=False))
 +AxisConfig=(AxisKeyName="OculusTouch_Left_Thumbstick",AxisProperties=(DeadZone=0.000000,Sensitivity=1.000000,Exponent=1.000000,bInvert=False))
 +AxisConfig=(AxisKeyName="OculusTouch_Left_FaceButton1",AxisProperties=(DeadZone=0.000000,Sensitivity=1.000000,Exponent=1.000000,bInvert=False))
 +AxisConfig=(AxisKeyName="OculusTouch_Left_Trigger",AxisProperties=(DeadZone=0.000000,Sensitivity=1.000000,Exponent=1.000000,bInvert=False))
@@ -53,13 +47,75 @@
 +AxisConfig=(AxisKeyName="SteamVR_Knuckles_Right_MiddleGrip",AxisProperties=(DeadZone=0.000000,Sensitivity=1.000000,Exponent=1.000000,bInvert=False))
 +AxisConfig=(AxisKeyName="SteamVR_Knuckles_Right_RingGrip",AxisProperties=(DeadZone=0.000000,Sensitivity=1.000000,Exponent=1.000000,bInvert=False))
 +AxisConfig=(AxisKeyName="SteamVR_Knuckles_Right_PinkyGrip",AxisProperties=(DeadZone=0.000000,Sensitivity=1.000000,Exponent=1.000000,bInvert=False))
++AxisConfig=(AxisKeyName="Daydream_Left_Trackpad_X",AxisProperties=(DeadZone=0.000000,Sensitivity=1.000000,Exponent=1.000000,bInvert=False))
++AxisConfig=(AxisKeyName="Daydream_Left_Trackpad_Y",AxisProperties=(DeadZone=0.000000,Sensitivity=1.000000,Exponent=1.000000,bInvert=False))
++AxisConfig=(AxisKeyName="Daydream_Right_Trackpad_X",AxisProperties=(DeadZone=0.000000,Sensitivity=1.000000,Exponent=1.000000,bInvert=False))
++AxisConfig=(AxisKeyName="Daydream_Right_Trackpad_Y",AxisProperties=(DeadZone=0.000000,Sensitivity=1.000000,Exponent=1.000000,bInvert=False))
++AxisConfig=(AxisKeyName="Vive_Left_Trigger_Axis",AxisProperties=(DeadZone=0.000000,Sensitivity=1.000000,Exponent=1.000000,bInvert=False))
++AxisConfig=(AxisKeyName="Vive_Left_Trackpad_X",AxisProperties=(DeadZone=0.000000,Sensitivity=1.000000,Exponent=1.000000,bInvert=False))
++AxisConfig=(AxisKeyName="Vive_Left_Trackpad_Y",AxisProperties=(DeadZone=0.000000,Sensitivity=1.000000,Exponent=1.000000,bInvert=False))
++AxisConfig=(AxisKeyName="Vive_Right_Trigger_Axis",AxisProperties=(DeadZone=0.000000,Sensitivity=1.000000,Exponent=1.000000,bInvert=False))
++AxisConfig=(AxisKeyName="Vive_Right_Trackpad_X",AxisProperties=(DeadZone=0.000000,Sensitivity=1.000000,Exponent=1.000000,bInvert=False))
++AxisConfig=(AxisKeyName="Vive_Right_Trackpad_Y",AxisProperties=(DeadZone=0.000000,Sensitivity=1.000000,Exponent=1.000000,bInvert=False))
++AxisConfig=(AxisKeyName="MixedReality_Left_Trigger_Axis",AxisProperties=(DeadZone=0.000000,Sensitivity=1.000000,Exponent=1.000000,bInvert=False))
++AxisConfig=(AxisKeyName="MixedReality_Left_Thumbstick_X",AxisProperties=(DeadZone=0.000000,Sensitivity=1.000000,Exponent=1.000000,bInvert=False))
++AxisConfig=(AxisKeyName="MixedReality_Left_Thumbstick_Y",AxisProperties=(DeadZone=0.000000,Sensitivity=1.000000,Exponent=1.000000,bInvert=False))
++AxisConfig=(AxisKeyName="MixedReality_Left_Trackpad_X",AxisProperties=(DeadZone=0.000000,Sensitivity=1.000000,Exponent=1.000000,bInvert=False))
++AxisConfig=(AxisKeyName="MixedReality_Left_Trackpad_Y",AxisProperties=(DeadZone=0.000000,Sensitivity=1.000000,Exponent=1.000000,bInvert=False))
++AxisConfig=(AxisKeyName="MixedReality_Right_Trigger_Axis",AxisProperties=(DeadZone=0.000000,Sensitivity=1.000000,Exponent=1.000000,bInvert=False))
++AxisConfig=(AxisKeyName="MixedReality_Right_Thumbstick_X",AxisProperties=(DeadZone=0.000000,Sensitivity=1.000000,Exponent=1.000000,bInvert=False))
++AxisConfig=(AxisKeyName="MixedReality_Right_Thumbstick_Y",AxisProperties=(DeadZone=0.000000,Sensitivity=1.000000,Exponent=1.000000,bInvert=False))
++AxisConfig=(AxisKeyName="MixedReality_Right_Trackpad_X",AxisProperties=(DeadZone=0.000000,Sensitivity=1.000000,Exponent=1.000000,bInvert=False))
++AxisConfig=(AxisKeyName="MixedReality_Right_Trackpad_Y",AxisProperties=(DeadZone=0.000000,Sensitivity=1.000000,Exponent=1.000000,bInvert=False))
++AxisConfig=(AxisKeyName="OculusGo_Left_Trackpad_X",AxisProperties=(DeadZone=0.000000,Sensitivity=1.000000,Exponent=1.000000,bInvert=False))
++AxisConfig=(AxisKeyName="OculusGo_Left_Trackpad_Y",AxisProperties=(DeadZone=0.000000,Sensitivity=1.000000,Exponent=1.000000,bInvert=False))
++AxisConfig=(AxisKeyName="OculusGo_Right_Trackpad_X",AxisProperties=(DeadZone=0.000000,Sensitivity=1.000000,Exponent=1.000000,bInvert=False))
++AxisConfig=(AxisKeyName="OculusGo_Right_Trackpad_Y",AxisProperties=(DeadZone=0.000000,Sensitivity=1.000000,Exponent=1.000000,bInvert=False))
++AxisConfig=(AxisKeyName="OculusTouch_Left_Grip_Axis",AxisProperties=(DeadZone=0.000000,Sensitivity=1.000000,Exponent=1.000000,bInvert=False))
++AxisConfig=(AxisKeyName="OculusTouch_Left_Trigger_Axis",AxisProperties=(DeadZone=0.000000,Sensitivity=1.000000,Exponent=1.000000,bInvert=False))
++AxisConfig=(AxisKeyName="OculusTouch_Left_Thumbstick_X",AxisProperties=(DeadZone=0.000000,Sensitivity=1.000000,Exponent=1.000000,bInvert=False))
++AxisConfig=(AxisKeyName="OculusTouch_Left_Thumbstick_Y",AxisProperties=(DeadZone=0.000000,Sensitivity=1.000000,Exponent=1.000000,bInvert=False))
++AxisConfig=(AxisKeyName="OculusTouch_Right_Grip_Axis",AxisProperties=(DeadZone=0.000000,Sensitivity=1.000000,Exponent=1.000000,bInvert=False))
++AxisConfig=(AxisKeyName="OculusTouch_Right_Trigger_Axis",AxisProperties=(DeadZone=0.000000,Sensitivity=1.000000,Exponent=1.000000,bInvert=False))
++AxisConfig=(AxisKeyName="OculusTouch_Right_Thumbstick_X",AxisProperties=(DeadZone=0.000000,Sensitivity=1.000000,Exponent=1.000000,bInvert=False))
++AxisConfig=(AxisKeyName="OculusTouch_Right_Thumbstick_Y",AxisProperties=(DeadZone=0.000000,Sensitivity=1.000000,Exponent=1.000000,bInvert=False))
++AxisConfig=(AxisKeyName="ValveIndex_Left_Grip_Axis",AxisProperties=(DeadZone=0.000000,Sensitivity=1.000000,Exponent=1.000000,bInvert=False))
++AxisConfig=(AxisKeyName="ValveIndex_Left_Grip_Force",AxisProperties=(DeadZone=0.000000,Sensitivity=1.000000,Exponent=1.000000,bInvert=False))
++AxisConfig=(AxisKeyName="ValveIndex_Left_Trigger_Axis",AxisProperties=(DeadZone=0.000000,Sensitivity=1.000000,Exponent=1.000000,bInvert=False))
++AxisConfig=(AxisKeyName="ValveIndex_Left_Thumbstick_X",AxisProperties=(DeadZone=0.000000,Sensitivity=1.000000,Exponent=1.000000,bInvert=False))
++AxisConfig=(AxisKeyName="ValveIndex_Left_Thumbstick_Y",AxisProperties=(DeadZone=0.000000,Sensitivity=1.000000,Exponent=1.000000,bInvert=False))
++AxisConfig=(AxisKeyName="ValveIndex_Left_Trackpad_X",AxisProperties=(DeadZone=0.000000,Sensitivity=1.000000,Exponent=1.000000,bInvert=False))
++AxisConfig=(AxisKeyName="ValveIndex_Left_Trackpad_Y",AxisProperties=(DeadZone=0.000000,Sensitivity=1.000000,Exponent=1.000000,bInvert=False))
++AxisConfig=(AxisKeyName="ValveIndex_Left_Trackpad_Force",AxisProperties=(DeadZone=0.000000,Sensitivity=1.000000,Exponent=1.000000,bInvert=False))
++AxisConfig=(AxisKeyName="ValveIndex_Left_Trackpad_Touch",AxisProperties=(DeadZone=0.000000,Sensitivity=1.000000,Exponent=1.000000,bInvert=False))
++AxisConfig=(AxisKeyName="ValveIndex_Right_Grip_Axis",AxisProperties=(DeadZone=0.000000,Sensitivity=1.000000,Exponent=1.000000,bInvert=False))
++AxisConfig=(AxisKeyName="ValveIndex_Right_Grip_Force",AxisProperties=(DeadZone=0.000000,Sensitivity=1.000000,Exponent=1.000000,bInvert=False))
++AxisConfig=(AxisKeyName="ValveIndex_Right_Trigger_Axis",AxisProperties=(DeadZone=0.000000,Sensitivity=1.000000,Exponent=1.000000,bInvert=False))
++AxisConfig=(AxisKeyName="ValveIndex_Right_Thumbstick_X",AxisProperties=(DeadZone=0.000000,Sensitivity=1.000000,Exponent=1.000000,bInvert=False))
++AxisConfig=(AxisKeyName="ValveIndex_Right_Thumbstick_Y",AxisProperties=(DeadZone=0.000000,Sensitivity=1.000000,Exponent=1.000000,bInvert=False))
++AxisConfig=(AxisKeyName="ValveIndex_Right_Trackpad_X",AxisProperties=(DeadZone=0.000000,Sensitivity=1.000000,Exponent=1.000000,bInvert=False))
++AxisConfig=(AxisKeyName="ValveIndex_Right_Trackpad_Y",AxisProperties=(DeadZone=0.000000,Sensitivity=1.000000,Exponent=1.000000,bInvert=False))
++AxisConfig=(AxisKeyName="ValveIndex_Right_Trackpad_Force",AxisProperties=(DeadZone=0.000000,Sensitivity=1.000000,Exponent=1.000000,bInvert=False))
++AxisConfig=(AxisKeyName="MagicLeap_Left_Trigger_Axis",AxisProperties=(DeadZone=0.000000,Sensitivity=1.000000,Exponent=1.000000,bInvert=False))
++AxisConfig=(AxisKeyName="MagicLeap_Left_Trackpad_X",AxisProperties=(DeadZone=0.000000,Sensitivity=1.000000,Exponent=1.000000,bInvert=False))
++AxisConfig=(AxisKeyName="MagicLeap_Left_Trackpad_Y",AxisProperties=(DeadZone=0.000000,Sensitivity=1.000000,Exponent=1.000000,bInvert=False))
++AxisConfig=(AxisKeyName="MagicLeap_Left_Trackpad_Force",AxisProperties=(DeadZone=0.000000,Sensitivity=1.000000,Exponent=1.000000,bInvert=False))
++AxisConfig=(AxisKeyName="MagicLeap_Left_Touch1_X",AxisProperties=(DeadZone=0.000000,Sensitivity=1.000000,Exponent=1.000000,bInvert=False))
++AxisConfig=(AxisKeyName="MagicLeap_Left_Touch1_Y",AxisProperties=(DeadZone=0.000000,Sensitivity=1.000000,Exponent=1.000000,bInvert=False))
++AxisConfig=(AxisKeyName="MagicLeap_Left_Touch1_Force",AxisProperties=(DeadZone=0.000000,Sensitivity=1.000000,Exponent=1.000000,bInvert=False))
++AxisConfig=(AxisKeyName="MagicLeap_Right_Trigger_Axis",AxisProperties=(DeadZone=0.000000,Sensitivity=1.000000,Exponent=1.000000,bInvert=False))
++AxisConfig=(AxisKeyName="MagicLeap_Right_Trackpad_X",AxisProperties=(DeadZone=0.000000,Sensitivity=1.000000,Exponent=1.000000,bInvert=False))
++AxisConfig=(AxisKeyName="MagicLeap_Right_Trackpad_Y",AxisProperties=(DeadZone=0.000000,Sensitivity=1.000000,Exponent=1.000000,bInvert=False))
++AxisConfig=(AxisKeyName="MagicLeap_Right_Trackpad_Force",AxisProperties=(DeadZone=0.000000,Sensitivity=1.000000,Exponent=1.000000,bInvert=False))
++AxisConfig=(AxisKeyName="MagicLeap_Right_Touch1_X",AxisProperties=(DeadZone=0.000000,Sensitivity=1.000000,Exponent=1.000000,bInvert=False))
++AxisConfig=(AxisKeyName="MagicLeap_Right_Touch1_Y",AxisProperties=(DeadZone=0.000000,Sensitivity=1.000000,Exponent=1.000000,bInvert=False))
++AxisConfig=(AxisKeyName="MagicLeap_Right_Touch1_Force",AxisProperties=(DeadZone=0.000000,Sensitivity=1.000000,Exponent=1.000000,bInvert=False))
 bAltEnterTogglesFullscreen=True
 bF11TogglesFullscreen=True
 bUseMouseForTouch=False
 bEnableMouseSmoothing=True
 bEnableFOVScaling=True
 bCaptureMouseOnLaunch=True
-bDefaultViewportMouseLock=True
 bAlwaysShowTouchInterface=False
 bShowConsoleOnFourFingerTap=True
 bEnableGestureRecognizer=False
@@ -70,13 +126,29 @@ FOVScale=0.011110
 DoubleClickTime=0.200000
 +ActionMappings=(ActionName="Jump",bShift=False,bCtrl=False,bAlt=False,bCmd=False,Key=SpaceBar)
 +ActionMappings=(ActionName="Jump",bShift=False,bCtrl=False,bAlt=False,bCmd=False,Key=Gamepad_FaceButton_Bottom)
++ActionMappings=(ActionName="Jump",bShift=False,bCtrl=False,bAlt=False,bCmd=False,Key=Daydream_Left_Select_Click)
 +ActionMappings=(ActionName="Fire",bShift=False,bCtrl=False,bAlt=False,bCmd=False,Key=LeftMouseButton)
 +ActionMappings=(ActionName="Fire",bShift=False,bCtrl=False,bAlt=False,bCmd=False,Key=Gamepad_RightTrigger)
-+ActionMappings=(ActionName="Fire",bShift=False,bCtrl=False,bAlt=False,bCmd=False,Key=MotionController_Right_Trigger)
-+ActionMappings=(ActionName="Jump",bShift=False,bCtrl=False,bAlt=False,bCmd=False,Key=MotionController_Left_Trigger)
++ActionMappings=(ActionName="Fire",bShift=False,bCtrl=False,bAlt=False,bCmd=False,Key=Daydream_Left_Trackpad_Click)
 +ActionMappings=(ActionName="ResetVR",bShift=False,bCtrl=False,bAlt=False,bCmd=False,Key=R)
-+ActionMappings=(ActionName="ResetVR",bShift=False,bCtrl=False,bAlt=False,bCmd=False,Key=MotionController_Left_Grip1)
-+ActionMappings=(ActionName="Fire",bShift=False,bCtrl=False,bAlt=False,bCmd=False,Key=OculusTouchpad_Touchpad)
++ActionMappings=(ActionName="ResetVR",bShift=False,bCtrl=False,bAlt=False,bCmd=False,Key=Vive_Left_Grip_Click)
++ActionMappings=(ActionName="Fire",bShift=False,bCtrl=False,bAlt=False,bCmd=False,Key=Vive_Right_Trigger_Click)
++ActionMappings=(ActionName="Jump",bShift=False,bCtrl=False,bAlt=False,bCmd=False,Key=Vive_Left_Trigger_Click)
++ActionMappings=(ActionName="Jump",bShift=False,bCtrl=False,bAlt=False,bCmd=False,Key=MixedReality_Left_Trigger_Click)
++ActionMappings=(ActionName="Jump",bShift=False,bCtrl=False,bAlt=False,bCmd=False,Key=OculusGo_Left_Trigger_Click)
++ActionMappings=(ActionName="Jump",bShift=False,bCtrl=False,bAlt=False,bCmd=False,Key=OculusTouch_Left_Trigger_Click)
++ActionMappings=(ActionName="Jump",bShift=False,bCtrl=False,bAlt=False,bCmd=False,Key=ValveIndex_Left_Trigger_Click)
++ActionMappings=(ActionName="Jump",bShift=False,bCtrl=False,bAlt=False,bCmd=False,Key=MagicLeap_Left_Trigger)
++ActionMappings=(ActionName="Fire",bShift=False,bCtrl=False,bAlt=False,bCmd=False,Key=MixedReality_Right_Trigger_Click)
++ActionMappings=(ActionName="Fire",bShift=False,bCtrl=False,bAlt=False,bCmd=False,Key=OculusGo_Left_Trackpad_Click)
++ActionMappings=(ActionName="Fire",bShift=False,bCtrl=False,bAlt=False,bCmd=False,Key=OculusTouch_Right_Trigger_Click)
++ActionMappings=(ActionName="Fire",bShift=False,bCtrl=False,bAlt=False,bCmd=False,Key=ValveIndex_Right_Trigger_Click)
++ActionMappings=(ActionName="Fire",bShift=False,bCtrl=False,bAlt=False,bCmd=False,Key=MagicLeap_Left_Trackpad_Touch)
++ActionMappings=(ActionName="ResetVR",bShift=False,bCtrl=False,bAlt=False,bCmd=False,Key=MixedReality_Left_Thumbstick_Click)
++ActionMappings=(ActionName="ResetVR",bShift=False,bCtrl=False,bAlt=False,bCmd=False,Key=OculusGo_Left_Back_Click)
++ActionMappings=(ActionName="ResetVR",bShift=False,bCtrl=False,bAlt=False,bCmd=False,Key=OculusTouch_Left_Thumbstick_Click)
++ActionMappings=(ActionName="ResetVR",bShift=False,bCtrl=False,bAlt=False,bCmd=False,Key=ValveIndex_Left_Thumbstick_Click)
++ActionMappings=(ActionName="ResetVR",bShift=False,bCtrl=False,bAlt=False,bCmd=False,Key=MagicLeap_Left_Bumper)
 +AxisMappings=(AxisName="MoveForward",Scale=1.000000,Key=W)
 +AxisMappings=(AxisName="MoveForward",Scale=-1.000000,Key=S)
 +AxisMappings=(AxisName="MoveForward",Scale=1.000000,Key=Up)
@@ -85,17 +157,27 @@ DoubleClickTime=0.200000
 +AxisMappings=(AxisName="MoveRight",Scale=-1.000000,Key=A)
 +AxisMappings=(AxisName="MoveRight",Scale=1.000000,Key=D)
 +AxisMappings=(AxisName="MoveRight",Scale=1.000000,Key=Gamepad_LeftX)
++AxisMappings=(AxisName="MoveRight",Scale=1.000000,Key=Daydream_Left_Trackpad_X)
 +AxisMappings=(AxisName="TurnRate",Scale=1.000000,Key=Gamepad_RightX)
 +AxisMappings=(AxisName="TurnRate",Scale=-1.000000,Key=Left)
 +AxisMappings=(AxisName="TurnRate",Scale=1.000000,Key=Right)
 +AxisMappings=(AxisName="Turn",Scale=1.000000,Key=MouseX)
 +AxisMappings=(AxisName="LookUpRate",Scale=1.000000,Key=Gamepad_RightY)
 +AxisMappings=(AxisName="LookUp",Scale=-1.000000,Key=MouseY)
-+AxisMappings=(AxisName="MoveForward",Scale=-1.000000,Key=MotionController_Left_Thumbstick_Y)
-+AxisMappings=(AxisName="MoveRight",Scale=1.000000,Key=MotionController_Left_Thumbstick_X)
++AxisMappings=(AxisName="MoveForward",Scale=1.000000,Key=Daydream_Left_Trackpad_Y)
++AxisMappings=(AxisName="MoveForward",Scale=1.000000,Key=Vive_Left_Trackpad_Y)
++AxisMappings=(AxisName="MoveForward",Scale=1.000000,Key=MixedReality_Left_Thumbstick_Y)
++AxisMappings=(AxisName="MoveForward",Scale=1.000000,Key=OculusGo_Left_Trackpad_Y)
++AxisMappings=(AxisName="MoveForward",Scale=1.000000,Key=OculusTouch_Left_Thumbstick_Y)
++AxisMappings=(AxisName="MoveForward",Scale=1.000000,Key=ValveIndex_Left_Thumbstick_Y)
++AxisMappings=(AxisName="MoveRight",Scale=1.000000,Key=Vive_Left_Trackpad_X)
++AxisMappings=(AxisName="MoveRight",Scale=1.000000,Key=MixedReality_Left_Thumbstick_X)
++AxisMappings=(AxisName="MoveRight",Scale=1.000000,Key=OculusGo_Left_Trackpad_X)
++AxisMappings=(AxisName="MoveRight",Scale=1.000000,Key=OculusTouch_Left_Thumbstick_Y)
++AxisMappings=(AxisName="MoveRight",Scale=1.000000,Key=ValveIndex_Left_Thumbstick_Y)
++AxisMappings=(AxisName="MoveRight",Scale=1.000000,Key=MagicLeap_Left_Trackpad_X)
++AxisMappings=(AxisName="MoveForward",Scale=1.000000,Key=MagicLeap_Left_Trackpad_Y)
 DefaultTouchInterface=/Engine/MobileResources/HUD/DefaultVirtualJoysticks.DefaultVirtualJoysticks
-ConsoleKey=None
--ConsoleKeys=Tilde
 +ConsoleKeys=Tilde
 
 
diff --git a/Templates/TP_Flying/Config/DefaultInput.ini b/Templates/TP_Flying/Config/DefaultInput.ini
--- a/Templates/TP_Flying/Config/DefaultInput.ini
+++ b/Templates/TP_Flying/Config/DefaultInput.ini
@@ -1,11 +1,5 @@
 
 [/Script/Engine.InputSettings]
--AxisConfig=(AxisKeyName="Gamepad_LeftX",AxisProperties=(DeadZone=0.25,Exponent=1.f,Sensitivity=1.f))
--AxisConfig=(AxisKeyName="Gamepad_LeftY",AxisProperties=(DeadZone=0.25,Exponent=1.f,Sensitivity=1.f))
--AxisConfig=(AxisKeyName="Gamepad_RightX",AxisProperties=(DeadZone=0.25,Exponent=1.f,Sensitivity=1.f))
--AxisConfig=(AxisKeyName="Gamepad_RightY",AxisProperties=(DeadZone=0.25,Exponent=1.f,Sensitivity=1.f))
--AxisConfig=(AxisKeyName="MouseX",AxisProperties=(DeadZone=0.f,Exponent=1.f,Sensitivity=0.07f))
--AxisConfig=(AxisKeyName="MouseY",AxisProperties=(DeadZone=0.f,Exponent=1.f,Sensitivity=0.07f))
 +AxisConfig=(AxisKeyName="Gamepad_LeftX",AxisProperties=(DeadZone=0.250000,Sensitivity=1.000000,Exponent=1.000000,bInvert=False))
 +AxisConfig=(AxisKeyName="Gamepad_LeftY",AxisProperties=(DeadZone=0.250000,Sensitivity=1.000000,Exponent=1.000000,bInvert=False))
 +AxisConfig=(AxisKeyName="Gamepad_RightX",AxisProperties=(DeadZone=0.250000,Sensitivity=1.000000,Exponent=1.000000,bInvert=False))
@@ -39,50 +33,134 @@
 +AxisConfig=(AxisKeyName="OculusTouch_Right_FaceButton2",AxisProperties=(DeadZone=0.000000,Sensitivity=1.000000,Exponent=1.000000,bInvert=False))
 +AxisConfig=(AxisKeyName="OculusTouch_Right_IndexPointing",AxisProperties=(DeadZone=0.000000,Sensitivity=1.000000,Exponent=1.000000,bInvert=False))
 +AxisConfig=(AxisKeyName="OculusTouch_Right_ThumbUp",AxisProperties=(DeadZone=0.000000,Sensitivity=1.000000,Exponent=1.000000,bInvert=False))
++AxisConfig=(AxisKeyName="Daydream_Left_Trackpad_X",AxisProperties=(DeadZone=0.000000,Sensitivity=1.000000,Exponent=1.000000,bInvert=False))
++AxisConfig=(AxisKeyName="Daydream_Left_Trackpad_Y",AxisProperties=(DeadZone=0.000000,Sensitivity=1.000000,Exponent=1.000000,bInvert=False))
++AxisConfig=(AxisKeyName="Daydream_Right_Trackpad_X",AxisProperties=(DeadZone=0.000000,Sensitivity=1.000000,Exponent=1.000000,bInvert=False))
++AxisConfig=(AxisKeyName="Daydream_Right_Trackpad_Y",AxisProperties=(DeadZone=0.000000,Sensitivity=1.000000,Exponent=1.000000,bInvert=False))
++AxisConfig=(AxisKeyName="Vive_Left_Trigger_Axis",AxisProperties=(DeadZone=0.000000,Sensitivity=1.000000,Exponent=1.000000,bInvert=False))
++AxisConfig=(AxisKeyName="Vive_Left_Trackpad_X",AxisProperties=(DeadZone=0.000000,Sensitivity=1.000000,Exponent=1.000000,bInvert=False))
++AxisConfig=(AxisKeyName="Vive_Left_Trackpad_Y",AxisProperties=(DeadZone=0.000000,Sensitivity=1.000000,Exponent=1.000000,bInvert=False))
++AxisConfig=(AxisKeyName="Vive_Right_Trigger_Axis",AxisProperties=(DeadZone=0.000000,Sensitivity=1.000000,Exponent=1.000000,bInvert=False))
++AxisConfig=(AxisKeyName="Vive_Right_Trackpad_X",AxisProperties=(DeadZone=0.000000,Sensitivity=1.000000,Exponent=1.000000,bInvert=False))
++AxisConfig=(AxisKeyName="Vive_Right_Trackpad_Y",AxisProperties=(DeadZone=0.000000,Sensitivity=1.000000,Exponent=1.000000,bInvert=False))
++AxisConfig=(AxisKeyName="MixedReality_Left_Trigger_Axis",AxisProperties=(DeadZone=0.000000,Sensitivity=1.000000,Exponent=1.000000,bInvert=False))
++AxisConfig=(AxisKeyName="MixedReality_Left_Thumbstick_X",AxisProperties=(DeadZone=0.000000,Sensitivity=1.000000,Exponent=1.000000,bInvert=False))
++AxisConfig=(AxisKeyName="MixedReality_Left_Thumbstick_Y",AxisProperties=(DeadZone=0.000000,Sensitivity=1.000000,Exponent=1.000000,bInvert=False))
++AxisConfig=(AxisKeyName="MixedReality_Left_Trackpad_X",AxisProperties=(DeadZone=0.000000,Sensitivity=1.000000,Exponent=1.000000,bInvert=False))
++AxisConfig=(AxisKeyName="MixedReality_Left_Trackpad_Y",AxisProperties=(DeadZone=0.000000,Sensitivity=1.000000,Exponent=1.000000,bInvert=False))
++AxisConfig=(AxisKeyName="MixedReality_Right_Trigger_Axis",AxisProperties=(DeadZone=0.000000,Sensitivity=1.000000,Exponent=1.000000,bInvert=False))
++AxisConfig=(AxisKeyName="MixedReality_Right_Thumbstick_X",AxisProperties=(DeadZone=0.000000,Sensitivity=1.000000,Exponent=1.000000,bInvert=False))
++AxisConfig=(AxisKeyName="MixedReality_Right_Thumbstick_Y",AxisProperties=(DeadZone=0.000000,Sensitivity=1.000000,Exponent=1.000000,bInvert=False))
++AxisConfig=(AxisKeyName="MixedReality_Right_Trackpad_X",AxisProperties=(DeadZone=0.000000,Sensitivity=1.000000,Exponent=1.000000,bInvert=False))
++AxisConfig=(AxisKeyName="MixedReality_Right_Trackpad_Y",AxisProperties=(DeadZone=0.000000,Sensitivity=1.000000,Exponent=1.000000,bInvert=False))
++AxisConfig=(AxisKeyName="OculusGo_Left_Trackpad_X",AxisProperties=(DeadZone=0.000000,Sensitivity=1.000000,Exponent=1.000000,bInvert=False))
++AxisConfig=(AxisKeyName="OculusGo_Left_Trackpad_Y",AxisProperties=(DeadZone=0.000000,Sensitivity=1.000000,Exponent=1.000000,bInvert=False))
++AxisConfig=(AxisKeyName="OculusGo_Right_Trackpad_X",AxisProperties=(DeadZone=0.000000,Sensitivity=1.000000,Exponent=1.000000,bInvert=False))
++AxisConfig=(AxisKeyName="OculusGo_Right_Trackpad_Y",AxisProperties=(DeadZone=0.000000,Sensitivity=1.000000,Exponent=1.000000,bInvert=False))
++AxisConfig=(AxisKeyName="OculusTouch_Left_Grip_Axis",AxisProperties=(DeadZone=0.000000,Sensitivity=1.000000,Exponent=1.000000,bInvert=False))
++AxisConfig=(AxisKeyName="OculusTouch_Left_Trigger_Axis",AxisProperties=(DeadZone=0.000000,Sensitivity=1.000000,Exponent=1.000000,bInvert=False))
++AxisConfig=(AxisKeyName="OculusTouch_Left_Thumbstick_X",AxisProperties=(DeadZone=0.000000,Sensitivity=1.000000,Exponent=1.000000,bInvert=False))
++AxisConfig=(AxisKeyName="OculusTouch_Left_Thumbstick_Y",AxisProperties=(DeadZone=0.000000,Sensitivity=1.000000,Exponent=1.000000,bInvert=False))
++AxisConfig=(AxisKeyName="OculusTouch_Right_Grip_Axis",AxisProperties=(DeadZone=0.000000,Sensitivity=1.000000,Exponent=1.000000,bInvert=False))
++AxisConfig=(AxisKeyName="OculusTouch_Right_Trigger_Axis",AxisProperties=(DeadZone=0.000000,Sensitivity=1.000000,Exponent=1.000000,bInvert=False))
++AxisConfig=(AxisKeyName="OculusTouch_Right_Thumbstick_X",AxisProperties=(DeadZone=0.000000,Sensitivity=1.000000,Exponent=1.000000,bInvert=False))
++AxisConfig=(AxisKeyName="OculusTouch_Right_Thumbstick_Y",AxisProperties=(DeadZone=0.000000,Sensitivity=1.000000,Exponent=1.000000,bInvert=False))
++AxisConfig=(AxisKeyName="ValveIndex_Left_Grip_Axis",AxisProperties=(DeadZone=0.000000,Sensitivity=1.000000,Exponent=1.000000,bInvert=False))
++AxisConfig=(AxisKeyName="ValveIndex_Left_Grip_Force",AxisProperties=(DeadZone=0.000000,Sensitivity=1.000000,Exponent=1.000000,bInvert=False))
++AxisConfig=(AxisKeyName="ValveIndex_Left_Trigger_Axis",AxisProperties=(DeadZone=0.000000,Sensitivity=1.000000,Exponent=1.000000,bInvert=False))
++AxisConfig=(AxisKeyName="ValveIndex_Left_Thumbstick_X",AxisProperties=(DeadZone=0.000000,Sensitivity=1.000000,Exponent=1.000000,bInvert=False))
++AxisConfig=(AxisKeyName="ValveIndex_Left_Thumbstick_Y",AxisProperties=(DeadZone=0.000000,Sensitivity=1.000000,Exponent=1.000000,bInvert=False))
++AxisConfig=(AxisKeyName="ValveIndex_Left_Trackpad_X",AxisProperties=(DeadZone=0.000000,Sensitivity=1.000000,Exponent=1.000000,bInvert=False))
++AxisConfig=(AxisKeyName="ValveIndex_Left_Trackpad_Y",AxisProperties=(DeadZone=0.000000,Sensitivity=1.000000,Exponent=1.000000,bInvert=False))
++AxisConfig=(AxisKeyName="ValveIndex_Left_Trackpad_Force",AxisProperties=(DeadZone=0.000000,Sensitivity=1.000000,Exponent=1.000000,bInvert=False))
++AxisConfig=(AxisKeyName="ValveIndex_Left_Trackpad_Touch",AxisProperties=(DeadZone=0.000000,Sensitivity=1.000000,Exponent=1.000000,bInvert=False))
++AxisConfig=(AxisKeyName="ValveIndex_Right_Grip_Axis",AxisProperties=(DeadZone=0.000000,Sensitivity=1.000000,Exponent=1.000000,bInvert=False))
++AxisConfig=(AxisKeyName="ValveIndex_Right_Grip_Force",AxisProperties=(DeadZone=0.000000,Sensitivity=1.000000,Exponent=1.000000,bInvert=False))
++AxisConfig=(AxisKeyName="ValveIndex_Right_Trigger_Axis",AxisProperties=(DeadZone=0.000000,Sensitivity=1.000000,Exponent=1.000000,bInvert=False))
++AxisConfig=(AxisKeyName="ValveIndex_Right_Thumbstick_X",AxisProperties=(DeadZone=0.000000,Sensitivity=1.000000,Exponent=1.000000,bInvert=False))
++AxisConfig=(AxisKeyName="ValveIndex_Right_Thumbstick_Y",AxisProperties=(DeadZone=0.000000,Sensitivity=1.000000,Exponent=1.000000,bInvert=False))
++AxisConfig=(AxisKeyName="ValveIndex_Right_Trackpad_X",AxisProperties=(DeadZone=0.000000,Sensitivity=1.000000,Exponent=1.000000,bInvert=False))
++AxisConfig=(AxisKeyName="ValveIndex_Right_Trackpad_Y",AxisProperties=(DeadZone=0.000000,Sensitivity=1.000000,Exponent=1.000000,bInvert=False))
++AxisConfig=(AxisKeyName="ValveIndex_Right_Trackpad_Force",AxisProperties=(DeadZone=0.000000,Sensitivity=1.000000,Exponent=1.000000,bInvert=False))
++AxisConfig=(AxisKeyName="MotionController_Left_Thumbstick_Z",AxisProperties=(DeadZone=0.000000,Sensitivity=1.000000,Exponent=1.000000,bInvert=False))
++AxisConfig=(AxisKeyName="MagicLeap_Left_Trigger_Axis",AxisProperties=(DeadZone=0.000000,Sensitivity=1.000000,Exponent=1.000000,bInvert=False))
++AxisConfig=(AxisKeyName="MagicLeap_Left_Trackpad_X",AxisProperties=(DeadZone=0.000000,Sensitivity=1.000000,Exponent=1.000000,bInvert=False))
++AxisConfig=(AxisKeyName="MagicLeap_Left_Trackpad_Y",AxisProperties=(DeadZone=0.000000,Sensitivity=1.000000,Exponent=1.000000,bInvert=False))
++AxisConfig=(AxisKeyName="MagicLeap_Left_Trackpad_Force",AxisProperties=(DeadZone=0.000000,Sensitivity=1.000000,Exponent=1.000000,bInvert=False))
++AxisConfig=(AxisKeyName="MagicLeap_Left_Touch1_X",AxisProperties=(DeadZone=0.000000,Sensitivity=1.000000,Exponent=1.000000,bInvert=False))
++AxisConfig=(AxisKeyName="MagicLeap_Left_Touch1_Y",AxisProperties=(DeadZone=0.000000,Sensitivity=1.000000,Exponent=1.000000,bInvert=False))
++AxisConfig=(AxisKeyName="MagicLeap_Left_Touch1_Force",AxisProperties=(DeadZone=0.000000,Sensitivity=1.000000,Exponent=1.000000,bInvert=False))
++AxisConfig=(AxisKeyName="MotionController_Right_Thumbstick_Z",AxisProperties=(DeadZone=0.000000,Sensitivity=1.000000,Exponent=1.000000,bInvert=False))
++AxisConfig=(AxisKeyName="MagicLeap_Right_Trigger_Axis",AxisProperties=(DeadZone=0.000000,Sensitivity=1.000000,Exponent=1.000000,bInvert=False))
++AxisConfig=(AxisKeyName="MagicLeap_Right_Trackpad_X",AxisProperties=(DeadZone=0.000000,Sensitivity=1.000000,Exponent=1.000000,bInvert=False))
++AxisConfig=(AxisKeyName="MagicLeap_Right_Trackpad_Y",AxisProperties=(DeadZone=0.000000,Sensitivity=1.000000,Exponent=1.000000,bInvert=False))
++AxisConfig=(AxisKeyName="MagicLeap_Right_Trackpad_Force",AxisProperties=(DeadZone=0.000000,Sensitivity=1.000000,Exponent=1.000000,bInvert=False))
++AxisConfig=(AxisKeyName="MagicLeap_Right_Touch1_X",AxisProperties=(DeadZone=0.000000,Sensitivity=1.000000,Exponent=1.000000,bInvert=False))
++AxisConfig=(AxisKeyName="MagicLeap_Right_Touch1_Y",AxisProperties=(DeadZone=0.000000,Sensitivity=1.000000,Exponent=1.000000,bInvert=False))
++AxisConfig=(AxisKeyName="MagicLeap_Right_Touch1_Force",AxisProperties=(DeadZone=0.000000,Sensitivity=1.000000,Exponent=1.000000,bInvert=False))
++AxisConfig=(AxisKeyName="OculusTouchpad_Touchpad_X",AxisProperties=(DeadZone=0.000000,Sensitivity=1.000000,Exponent=1.000000,bInvert=False))
++AxisConfig=(AxisKeyName="OculusTouchpad_Touchpad_Y",AxisProperties=(DeadZone=0.000000,Sensitivity=1.000000,Exponent=1.000000,bInvert=False))
 bAltEnterTogglesFullscreen=True
 bF11TogglesFullscreen=True
 bUseMouseForTouch=False
 bEnableMouseSmoothing=True
 bEnableFOVScaling=True
-FOVScale=0.011110
-DoubleClickTime=0.200000
 bCaptureMouseOnLaunch=True
-DefaultViewportMouseCaptureMode=CapturePermanently_IncludingInitialMouseDown
-bDefaultViewportMouseLock=True
-DefaultViewportMouseLockMode=LockOnCapture
-+AxisMappings=(AxisName="MoveUp",Key=W,Scale=1.000000)
-+AxisMappings=(AxisName="MoveUp",Key=S,Scale=-1.000000)
-+AxisMappings=(AxisName="MoveUp",Key=Gamepad_LeftStick_Up,Scale=1.000000)
-+AxisMappings=(AxisName="MoveUp",Key=Gamepad_LeftStick_Down,Scale=-1.000000)
-+AxisMappings=(AxisName="MoveUp",Key=Up,Scale=1.000000)
-+AxisMappings=(AxisName="MoveUp",Key=Down,Scale=-1.000000)
-+AxisMappings=(AxisName="MoveRight",Key=D,Scale=1.000000)
-+AxisMappings=(AxisName="MoveRight",Key=A,Scale=-1.000000)
-+AxisMappings=(AxisName="MoveRight",Key=Gamepad_LeftStick_Right,Scale=1.000000)
-+AxisMappings=(AxisName="MoveRight",Key=Gamepad_LeftStick_Left,Scale=-1.000000)
-+AxisMappings=(AxisName="MoveRight",Key=Right,Scale=1.000000)
-+AxisMappings=(AxisName="MoveRight",Key=Left,Scale=-1.000000)
-+AxisMappings=(AxisName="Thrust",Key=LeftShift,Scale=1.000000)
-+AxisMappings=(AxisName="Thrust",Key=RightShift,Scale=1.000000)
-+AxisMappings=(AxisName="Thrust",Key=LeftControl,Scale=-1.000000)
-+AxisMappings=(AxisName="Thrust",Key=RightControl,Scale=-1.000000)
-+AxisMappings=(AxisName="MoveUp",Key=Gamepad_LeftY,Scale=1.000000)
-+AxisMappings=(AxisName="MoveRight",Key=Gamepad_LeftX,Scale=1.000000)
-+AxisMappings=(AxisName="Thrust",Key=SpaceBar,Scale=1.000000)
-+AxisMappings=(AxisName="MoveUp",Key=MotionController_Left_Thumbstick_Y,Scale=1.000000)
-+AxisMappings=(AxisName="MoveUp",Key=MotionController_Right_Thumbstick_Y,Scale=1.000000)
-+AxisMappings=(AxisName="MoveRight",Key=MotionController_Left_Thumbstick_X,Scale=1.000000)
-+AxisMappings=(AxisName="MoveRight",Key=MotionController_Right_Thumbstick_X,Scale=1.000000)
-+AxisMappings=(AxisName="Thrust",Key=MotionController_Right_Trigger,Scale=1.000000)
-+AxisMappings=(AxisName="Thrust",Key=MotionController_Left_Trigger,Scale=-1.000000)
-+AxisMappings=(AxisName="Thrust",Key=Gamepad_RightStick_Up,Scale=-1.000000)
-+AxisMappings=(AxisName="Thrust",Key=Gamepad_RightStick_Down,Scale=1.000000)
-+AxisMappings=(AxisName="Thrust",Key=Gamepad_RightY,Scale=1.000000)
 bAlwaysShowTouchInterface=False
 bShowConsoleOnFourFingerTap=True
+bEnableGestureRecognizer=False
+bUseAutocorrect=False
+DefaultViewportMouseCaptureMode=CapturePermanently_IncludingInitialMouseDown
+DefaultViewportMouseLockMode=LockOnCapture
+FOVScale=0.011110
+DoubleClickTime=0.200000
++AxisMappings=(AxisName="MoveUp",Scale=1.000000,Key=W)
++AxisMappings=(AxisName="MoveUp",Scale=-1.000000,Key=S)
++AxisMappings=(AxisName="MoveUp",Scale=1.000000,Key=Gamepad_LeftStick_Up)
++AxisMappings=(AxisName="MoveUp",Scale=-1.000000,Key=Gamepad_LeftStick_Down)
++AxisMappings=(AxisName="MoveUp",Scale=1.000000,Key=Up)
++AxisMappings=(AxisName="MoveUp",Scale=-1.000000,Key=Down)
++AxisMappings=(AxisName="MoveRight",Scale=1.000000,Key=D)
++AxisMappings=(AxisName="MoveRight",Scale=-1.000000,Key=A)
++AxisMappings=(AxisName="MoveRight",Scale=1.000000,Key=Gamepad_LeftStick_Right)
++AxisMappings=(AxisName="MoveRight",Scale=-1.000000,Key=Gamepad_LeftStick_Left)
++AxisMappings=(AxisName="MoveRight",Scale=1.000000,Key=Right)
++AxisMappings=(AxisName="MoveRight",Scale=-1.000000,Key=Left)
++AxisMappings=(AxisName="Thrust",Scale=1.000000,Key=LeftShift)
++AxisMappings=(AxisName="Thrust",Scale=1.000000,Key=RightShift)
++AxisMappings=(AxisName="Thrust",Scale=-1.000000,Key=LeftControl)
++AxisMappings=(AxisName="Thrust",Scale=-1.000000,Key=RightControl)
++AxisMappings=(AxisName="MoveUp",Scale=1.000000,Key=Gamepad_LeftY)
++AxisMappings=(AxisName="MoveRight",Scale=1.000000,Key=Gamepad_LeftX)
++AxisMappings=(AxisName="Thrust",Scale=1.000000,Key=SpaceBar)
++AxisMappings=(AxisName="Thrust",Scale=-1.000000,Key=Gamepad_RightStick_Up)
++AxisMappings=(AxisName="Thrust",Scale=1.000000,Key=Gamepad_RightStick_Down)
++AxisMappings=(AxisName="Thrust",Scale=1.000000,Key=Gamepad_RightY)
++AxisMappings=(AxisName="Thrust",Scale=1.000000,Key=Vive_Left_Trigger_Axis)
++AxisMappings=(AxisName="Thrust",Scale=-1.000000,Key=Vive_Right_Trigger_Axis)
++AxisMappings=(AxisName="Thrust",Scale=1.000000,Key=MixedReality_Left_Trigger_Axis)
++AxisMappings=(AxisName="Thrust",Scale=-1.000000,Key=MixedReality_Right_Trigger_Axis)
++AxisMappings=(AxisName="Thrust",Scale=1.000000,Key=OculusTouch_Left_Trigger_Axis)
++AxisMappings=(AxisName="Thrust",Scale=-1.000000,Key=OculusTouch_Right_Trigger_Axis)
++AxisMappings=(AxisName="Thrust",Scale=1.000000,Key=ValveIndex_Left_Trigger_Axis)
++AxisMappings=(AxisName="Thrust",Scale=-1.000000,Key=ValveIndex_Right_Trigger_Axis)
++AxisMappings=(AxisName="MoveRight",Scale=1.000000,Key=Vive_Left_Trackpad_X)
++AxisMappings=(AxisName="MoveRight",Scale=1.000000,Key=Vive_Right_Trackpad_X)
++AxisMappings=(AxisName="MoveRight",Scale=1.000000,Key=MixedReality_Left_Thumbstick_X)
++AxisMappings=(AxisName="MoveRight",Scale=1.000000,Key=MixedReality_Right_Thumbstick_X)
++AxisMappings=(AxisName="MoveRight",Scale=1.000000,Key=OculusTouch_Left_Thumbstick_X)
++AxisMappings=(AxisName="MoveRight",Scale=1.000000,Key=OculusTouch_Right_Thumbstick_X)
++AxisMappings=(AxisName="MoveRight",Scale=1.000000,Key=ValveIndex_Left_Thumbstick_X)
++AxisMappings=(AxisName="MoveRight",Scale=1.000000,Key=ValveIndex_Right_Thumbstick_X)
++AxisMappings=(AxisName="MoveUp",Scale=1.000000,Key=Vive_Left_Trackpad_Y)
++AxisMappings=(AxisName="MoveUp",Scale=1.000000,Key=Vive_Right_Trackpad_Y)
++AxisMappings=(AxisName="MoveUp",Scale=1.000000,Key=MixedReality_Left_Thumbstick_Y)
++AxisMappings=(AxisName="MoveUp",Scale=1.000000,Key=MixedReality_Right_Thumbstick_Y)
++AxisMappings=(AxisName="MoveUp",Scale=1.000000,Key=OculusTouch_Left_Thumbstick_Y)
++AxisMappings=(AxisName="MoveUp",Scale=1.000000,Key=OculusTouch_Right_Thumbstick_Y)
++AxisMappings=(AxisName="MoveUp",Scale=1.000000,Key=ValveIndex_Left_Thumbstick_Y)
++AxisMappings=(AxisName="MoveUp",Scale=1.000000,Key=ValveIndex_Right_Thumbstick_Y)
 DefaultTouchInterface=/Engine/MobileResources/HUD/DefaultVirtualJoysticks.DefaultVirtualJoysticks
-ConsoleKey=None
--ConsoleKeys=Tilde
 +ConsoleKeys=Tilde
 
 
diff --git a/Templates/TP_FlyingBP/Config/DefaultInput.ini b/Templates/TP_FlyingBP/Config/DefaultInput.ini
--- a/Templates/TP_FlyingBP/Config/DefaultInput.ini
+++ b/Templates/TP_FlyingBP/Config/DefaultInput.ini
@@ -1,11 +1,5 @@
 
 [/Script/Engine.InputSettings]
--AxisConfig=(AxisKeyName="Gamepad_LeftX",AxisProperties=(DeadZone=0.25,Exponent=1.f,Sensitivity=1.f))
--AxisConfig=(AxisKeyName="Gamepad_LeftY",AxisProperties=(DeadZone=0.25,Exponent=1.f,Sensitivity=1.f))
--AxisConfig=(AxisKeyName="Gamepad_RightX",AxisProperties=(DeadZone=0.25,Exponent=1.f,Sensitivity=1.f))
--AxisConfig=(AxisKeyName="Gamepad_RightY",AxisProperties=(DeadZone=0.25,Exponent=1.f,Sensitivity=1.f))
--AxisConfig=(AxisKeyName="MouseX",AxisProperties=(DeadZone=0.f,Exponent=1.f,Sensitivity=0.07f))
--AxisConfig=(AxisKeyName="MouseY",AxisProperties=(DeadZone=0.f,Exponent=1.f,Sensitivity=0.07f))
 +AxisConfig=(AxisKeyName="Gamepad_LeftX",AxisProperties=(DeadZone=0.250000,Sensitivity=1.000000,Exponent=1.000000,bInvert=False))
 +AxisConfig=(AxisKeyName="Gamepad_LeftY",AxisProperties=(DeadZone=0.250000,Sensitivity=1.000000,Exponent=1.000000,bInvert=False))
 +AxisConfig=(AxisKeyName="Gamepad_RightX",AxisProperties=(DeadZone=0.250000,Sensitivity=1.000000,Exponent=1.000000,bInvert=False))
@@ -39,50 +33,134 @@
 +AxisConfig=(AxisKeyName="OculusTouch_Right_FaceButton2",AxisProperties=(DeadZone=0.000000,Sensitivity=1.000000,Exponent=1.000000,bInvert=False))
 +AxisConfig=(AxisKeyName="OculusTouch_Right_IndexPointing",AxisProperties=(DeadZone=0.000000,Sensitivity=1.000000,Exponent=1.000000,bInvert=False))
 +AxisConfig=(AxisKeyName="OculusTouch_Right_ThumbUp",AxisProperties=(DeadZone=0.000000,Sensitivity=1.000000,Exponent=1.000000,bInvert=False))
++AxisConfig=(AxisKeyName="Daydream_Left_Trackpad_X",AxisProperties=(DeadZone=0.000000,Sensitivity=1.000000,Exponent=1.000000,bInvert=False))
++AxisConfig=(AxisKeyName="Daydream_Left_Trackpad_Y",AxisProperties=(DeadZone=0.000000,Sensitivity=1.000000,Exponent=1.000000,bInvert=False))
++AxisConfig=(AxisKeyName="Daydream_Right_Trackpad_X",AxisProperties=(DeadZone=0.000000,Sensitivity=1.000000,Exponent=1.000000,bInvert=False))
++AxisConfig=(AxisKeyName="Daydream_Right_Trackpad_Y",AxisProperties=(DeadZone=0.000000,Sensitivity=1.000000,Exponent=1.000000,bInvert=False))
++AxisConfig=(AxisKeyName="Vive_Left_Trigger_Axis",AxisProperties=(DeadZone=0.000000,Sensitivity=1.000000,Exponent=1.000000,bInvert=False))
++AxisConfig=(AxisKeyName="Vive_Left_Trackpad_X",AxisProperties=(DeadZone=0.000000,Sensitivity=1.000000,Exponent=1.000000,bInvert=False))
++AxisConfig=(AxisKeyName="Vive_Left_Trackpad_Y",AxisProperties=(DeadZone=0.000000,Sensitivity=1.000000,Exponent=1.000000,bInvert=False))
++AxisConfig=(AxisKeyName="Vive_Right_Trigger_Axis",AxisProperties=(DeadZone=0.000000,Sensitivity=1.000000,Exponent=1.000000,bInvert=False))
++AxisConfig=(AxisKeyName="Vive_Right_Trackpad_X",AxisProperties=(DeadZone=0.000000,Sensitivity=1.000000,Exponent=1.000000,bInvert=False))
++AxisConfig=(AxisKeyName="Vive_Right_Trackpad_Y",AxisProperties=(DeadZone=0.000000,Sensitivity=1.000000,Exponent=1.000000,bInvert=False))
++AxisConfig=(AxisKeyName="MixedReality_Left_Trigger_Axis",AxisProperties=(DeadZone=0.000000,Sensitivity=1.000000,Exponent=1.000000,bInvert=False))
++AxisConfig=(AxisKeyName="MixedReality_Left_Thumbstick_X",AxisProperties=(DeadZone=0.000000,Sensitivity=1.000000,Exponent=1.000000,bInvert=False))
++AxisConfig=(AxisKeyName="MixedReality_Left_Thumbstick_Y",AxisProperties=(DeadZone=0.000000,Sensitivity=1.000000,Exponent=1.000000,bInvert=False))
++AxisConfig=(AxisKeyName="MixedReality_Left_Trackpad_X",AxisProperties=(DeadZone=0.000000,Sensitivity=1.000000,Exponent=1.000000,bInvert=False))
++AxisConfig=(AxisKeyName="MixedReality_Left_Trackpad_Y",AxisProperties=(DeadZone=0.000000,Sensitivity=1.000000,Exponent=1.000000,bInvert=False))
++AxisConfig=(AxisKeyName="MixedReality_Right_Trigger_Axis",AxisProperties=(DeadZone=0.000000,Sensitivity=1.000000,Exponent=1.000000,bInvert=False))
++AxisConfig=(AxisKeyName="MixedReality_Right_Thumbstick_X",AxisProperties=(DeadZone=0.000000,Sensitivity=1.000000,Exponent=1.000000,bInvert=False))
++AxisConfig=(AxisKeyName="MixedReality_Right_Thumbstick_Y",AxisProperties=(DeadZone=0.000000,Sensitivity=1.000000,Exponent=1.000000,bInvert=False))
++AxisConfig=(AxisKeyName="MixedReality_Right_Trackpad_X",AxisProperties=(DeadZone=0.000000,Sensitivity=1.000000,Exponent=1.000000,bInvert=False))
++AxisConfig=(AxisKeyName="MixedReality_Right_Trackpad_Y",AxisProperties=(DeadZone=0.000000,Sensitivity=1.000000,Exponent=1.000000,bInvert=False))
++AxisConfig=(AxisKeyName="OculusGo_Left_Trackpad_X",AxisProperties=(DeadZone=0.000000,Sensitivity=1.000000,Exponent=1.000000,bInvert=False))
++AxisConfig=(AxisKeyName="OculusGo_Left_Trackpad_Y",AxisProperties=(DeadZone=0.000000,Sensitivity=1.000000,Exponent=1.000000,bInvert=False))
++AxisConfig=(AxisKeyName="OculusGo_Right_Trackpad_X",AxisProperties=(DeadZone=0.000000,Sensitivity=1.000000,Exponent=1.000000,bInvert=False))
++AxisConfig=(AxisKeyName="OculusGo_Right_Trackpad_Y",AxisProperties=(DeadZone=0.000000,Sensitivity=1.000000,Exponent=1.000000,bInvert=False))
++AxisConfig=(AxisKeyName="OculusTouch_Left_Grip_Axis",AxisProperties=(DeadZone=0.000000,Sensitivity=1.000000,Exponent=1.000000,bInvert=False))
++AxisConfig=(AxisKeyName="OculusTouch_Left_Trigger_Axis",AxisProperties=(DeadZone=0.000000,Sensitivity=1.000000,Exponent=1.000000,bInvert=False))
++AxisConfig=(AxisKeyName="OculusTouch_Left_Thumbstick_X",AxisProperties=(DeadZone=0.000000,Sensitivity=1.000000,Exponent=1.000000,bInvert=False))
++AxisConfig=(AxisKeyName="OculusTouch_Left_Thumbstick_Y",AxisProperties=(DeadZone=0.000000,Sensitivity=1.000000,Exponent=1.000000,bInvert=False))
++AxisConfig=(AxisKeyName="OculusTouch_Right_Grip_Axis",AxisProperties=(DeadZone=0.000000,Sensitivity=1.000000,Exponent=1.000000,bInvert=False))
++AxisConfig=(AxisKeyName="OculusTouch_Right_Trigger_Axis",AxisProperties=(DeadZone=0.000000,Sensitivity=1.000000,Exponent=1.000000,bInvert=False))
++AxisConfig=(AxisKeyName="OculusTouch_Right_Thumbstick_X",AxisProperties=(DeadZone=0.000000,Sensitivity=1.000000,Exponent=1.000000,bInvert=False))
++AxisConfig=(AxisKeyName="OculusTouch_Right_Thumbstick_Y",AxisProperties=(DeadZone=0.000000,Sensitivity=1.000000,Exponent=1.000000,bInvert=False))
++AxisConfig=(AxisKeyName="ValveIndex_Left_Grip_Axis",AxisProperties=(DeadZone=0.000000,Sensitivity=1.000000,Exponent=1.000000,bInvert=False))
++AxisConfig=(AxisKeyName="ValveIndex_Left_Grip_Force",AxisProperties=(DeadZone=0.000000,Sensitivity=1.000000,Exponent=1.000000,bInvert=False))
++AxisConfig=(AxisKeyName="ValveIndex_Left_Trigger_Axis",AxisProperties=(DeadZone=0.000000,Sensitivity=1.000000,Exponent=1.000000,bInvert=False))
++AxisConfig=(AxisKeyName="ValveIndex_Left_Thumbstick_X",AxisProperties=(DeadZone=0.000000,Sensitivity=1.000000,Exponent=1.000000,bInvert=False))
++AxisConfig=(AxisKeyName="ValveIndex_Left_Thumbstick_Y",AxisProperties=(DeadZone=0.000000,Sensitivity=1.000000,Exponent=1.000000,bInvert=False))
++AxisConfig=(AxisKeyName="ValveIndex_Left_Trackpad_X",AxisProperties=(DeadZone=0.000000,Sensitivity=1.000000,Exponent=1.000000,bInvert=False))
++AxisConfig=(AxisKeyName="ValveIndex_Left_Trackpad_Y",AxisProperties=(DeadZone=0.000000,Sensitivity=1.000000,Exponent=1.000000,bInvert=False))
++AxisConfig=(AxisKeyName="ValveIndex_Left_Trackpad_Force",AxisProperties=(DeadZone=0.000000,Sensitivity=1.000000,Exponent=1.000000,bInvert=False))
++AxisConfig=(AxisKeyName="ValveIndex_Left_Trackpad_Touch",AxisProperties=(DeadZone=0.000000,Sensitivity=1.000000,Exponent=1.000000,bInvert=False))
++AxisConfig=(AxisKeyName="ValveIndex_Right_Grip_Axis",AxisProperties=(DeadZone=0.000000,Sensitivity=1.000000,Exponent=1.000000,bInvert=False))
++AxisConfig=(AxisKeyName="ValveIndex_Right_Grip_Force",AxisProperties=(DeadZone=0.000000,Sensitivity=1.000000,Exponent=1.000000,bInvert=False))
++AxisConfig=(AxisKeyName="ValveIndex_Right_Trigger_Axis",AxisProperties=(DeadZone=0.000000,Sensitivity=1.000000,Exponent=1.000000,bInvert=False))
++AxisConfig=(AxisKeyName="ValveIndex_Right_Thumbstick_X",AxisProperties=(DeadZone=0.000000,Sensitivity=1.000000,Exponent=1.000000,bInvert=False))
++AxisConfig=(AxisKeyName="ValveIndex_Right_Thumbstick_Y",AxisProperties=(DeadZone=0.000000,Sensitivity=1.000000,Exponent=1.000000,bInvert=False))
++AxisConfig=(AxisKeyName="ValveIndex_Right_Trackpad_X",AxisProperties=(DeadZone=0.000000,Sensitivity=1.000000,Exponent=1.000000,bInvert=False))
++AxisConfig=(AxisKeyName="ValveIndex_Right_Trackpad_Y",AxisProperties=(DeadZone=0.000000,Sensitivity=1.000000,Exponent=1.000000,bInvert=False))
++AxisConfig=(AxisKeyName="ValveIndex_Right_Trackpad_Force",AxisProperties=(DeadZone=0.000000,Sensitivity=1.000000,Exponent=1.000000,bInvert=False))
++AxisConfig=(AxisKeyName="MotionController_Left_Thumbstick_Z",AxisProperties=(DeadZone=0.000000,Sensitivity=1.000000,Exponent=1.000000,bInvert=False))
++AxisConfig=(AxisKeyName="MagicLeap_Left_Trigger_Axis",AxisProperties=(DeadZone=0.000000,Sensitivity=1.000000,Exponent=1.000000,bInvert=False))
++AxisConfig=(AxisKeyName="MagicLeap_Left_Trackpad_X",AxisProperties=(DeadZone=0.000000,Sensitivity=1.000000,Exponent=1.000000,bInvert=False))
++AxisConfig=(AxisKeyName="MagicLeap_Left_Trackpad_Y",AxisProperties=(DeadZone=0.000000,Sensitivity=1.000000,Exponent=1.000000,bInvert=False))
++AxisConfig=(AxisKeyName="MagicLeap_Left_Trackpad_Force",AxisProperties=(DeadZone=0.000000,Sensitivity=1.000000,Exponent=1.000000,bInvert=False))
++AxisConfig=(AxisKeyName="MagicLeap_Left_Touch1_X",AxisProperties=(DeadZone=0.000000,Sensitivity=1.000000,Exponent=1.000000,bInvert=False))
++AxisConfig=(AxisKeyName="MagicLeap_Left_Touch1_Y",AxisProperties=(DeadZone=0.000000,Sensitivity=1.000000,Exponent=1.000000,bInvert=False))
++AxisConfig=(AxisKeyName="MagicLeap_Left_Touch1_Force",AxisProperties=(DeadZone=0.000000,Sensitivity=1.000000,Exponent=1.000000,bInvert=False))
++AxisConfig=(AxisKeyName="MotionController_Right_Thumbstick_Z",AxisProperties=(DeadZone=0.000000,Sensitivity=1.000000,Exponent=1.000000,bInvert=False))
++AxisConfig=(AxisKeyName="MagicLeap_Right_Trigger_Axis",AxisProperties=(DeadZone=0.000000,Sensitivity=1.000000,Exponent=1.000000,bInvert=False))
++AxisConfig=(AxisKeyName="MagicLeap_Right_Trackpad_X",AxisProperties=(DeadZone=0.000000,Sensitivity=1.000000,Exponent=1.000000,bInvert=False))
++AxisConfig=(AxisKeyName="MagicLeap_Right_Trackpad_Y",AxisProperties=(DeadZone=0.000000,Sensitivity=1.000000,Exponent=1.000000,bInvert=False))
++AxisConfig=(AxisKeyName="MagicLeap_Right_Trackpad_Force",AxisProperties=(DeadZone=0.000000,Sensitivity=1.000000,Exponent=1.000000,bInvert=False))
++AxisConfig=(AxisKeyName="MagicLeap_Right_Touch1_X",AxisProperties=(DeadZone=0.000000,Sensitivity=1.000000,Exponent=1.000000,bInvert=False))
++AxisConfig=(AxisKeyName="MagicLeap_Right_Touch1_Y",AxisProperties=(DeadZone=0.000000,Sensitivity=1.000000,Exponent=1.000000,bInvert=False))
++AxisConfig=(AxisKeyName="MagicLeap_Right_Touch1_Force",AxisProperties=(DeadZone=0.000000,Sensitivity=1.000000,Exponent=1.000000,bInvert=False))
++AxisConfig=(AxisKeyName="OculusTouchpad_Touchpad_X",AxisProperties=(DeadZone=0.000000,Sensitivity=1.000000,Exponent=1.000000,bInvert=False))
++AxisConfig=(AxisKeyName="OculusTouchpad_Touchpad_Y",AxisProperties=(DeadZone=0.000000,Sensitivity=1.000000,Exponent=1.000000,bInvert=False))
 bAltEnterTogglesFullscreen=True
 bF11TogglesFullscreen=True
 bUseMouseForTouch=False
 bEnableMouseSmoothing=True
 bEnableFOVScaling=True
-FOVScale=0.011110
-DoubleClickTime=0.200000
 bCaptureMouseOnLaunch=True
-DefaultViewportMouseCaptureMode=CapturePermanently_IncludingInitialMouseDown
-bDefaultViewportMouseLock=True
-DefaultViewportMouseLockMode=LockOnCapture
-+AxisMappings=(AxisName="MoveUp",Key=W,Scale=1.000000)
-+AxisMappings=(AxisName="MoveUp",Key=S,Scale=-1.000000)
-+AxisMappings=(AxisName="MoveUp",Key=Gamepad_LeftStick_Up,Scale=1.000000)
-+AxisMappings=(AxisName="MoveUp",Key=Gamepad_LeftStick_Down,Scale=-1.000000)
-+AxisMappings=(AxisName="MoveUp",Key=Up,Scale=1.000000)
-+AxisMappings=(AxisName="MoveUp",Key=Down,Scale=-1.000000)
-+AxisMappings=(AxisName="MoveRight",Key=D,Scale=1.000000)
-+AxisMappings=(AxisName="MoveRight",Key=A,Scale=-1.000000)
-+AxisMappings=(AxisName="MoveRight",Key=Gamepad_LeftStick_Right,Scale=1.000000)
-+AxisMappings=(AxisName="MoveRight",Key=Gamepad_LeftStick_Left,Scale=-1.000000)
-+AxisMappings=(AxisName="MoveRight",Key=Right,Scale=1.000000)
-+AxisMappings=(AxisName="MoveRight",Key=Left,Scale=-1.000000)
-+AxisMappings=(AxisName="Thrust",Key=LeftShift,Scale=1.000000)
-+AxisMappings=(AxisName="Thrust",Key=RightShift,Scale=1.000000)
-+AxisMappings=(AxisName="Thrust",Key=LeftControl,Scale=-1.000000)
-+AxisMappings=(AxisName="Thrust",Key=RightControl,Scale=-1.000000)
-+AxisMappings=(AxisName="MoveUp",Key=Gamepad_LeftY,Scale=1.000000)
-+AxisMappings=(AxisName="MoveRight",Key=Gamepad_LeftX,Scale=1.000000)
-+AxisMappings=(AxisName="Thrust",Key=SpaceBar,Scale=1.000000)
-+AxisMappings=(AxisName="MoveUp",Key=MotionController_Left_Thumbstick_Y,Scale=1.000000)
-+AxisMappings=(AxisName="MoveUp",Key=MotionController_Right_Thumbstick_Y,Scale=1.000000)
-+AxisMappings=(AxisName="MoveRight",Key=MotionController_Left_Thumbstick_X,Scale=1.000000)
-+AxisMappings=(AxisName="MoveRight",Key=MotionController_Right_Thumbstick_X,Scale=1.000000)
-+AxisMappings=(AxisName="Thrust",Key=MotionController_Right_Trigger,Scale=1.000000)
-+AxisMappings=(AxisName="Thrust",Key=MotionController_Left_Trigger,Scale=-1.000000)
-+AxisMappings=(AxisName="Thrust",Key=Gamepad_RightY,Scale=1.000000)
-+AxisMappings=(AxisName="Thrust",Key=Gamepad_RightStick_Up,Scale=-1.000000)
-+AxisMappings=(AxisName="Thrust",Key=Gamepad_RightStick_Down,Scale=1.000000)
 bAlwaysShowTouchInterface=False
 bShowConsoleOnFourFingerTap=True
+bEnableGestureRecognizer=False
+bUseAutocorrect=False
+DefaultViewportMouseCaptureMode=CapturePermanently_IncludingInitialMouseDown
+DefaultViewportMouseLockMode=LockOnCapture
+FOVScale=0.011110
+DoubleClickTime=0.200000
++AxisMappings=(AxisName="MoveUp",Scale=1.000000,Key=W)
++AxisMappings=(AxisName="MoveUp",Scale=-1.000000,Key=S)
++AxisMappings=(AxisName="MoveUp",Scale=1.000000,Key=Gamepad_LeftStick_Up)
++AxisMappings=(AxisName="MoveUp",Scale=-1.000000,Key=Gamepad_LeftStick_Down)
++AxisMappings=(AxisName="MoveUp",Scale=1.000000,Key=Up)
++AxisMappings=(AxisName="MoveUp",Scale=-1.000000,Key=Down)
++AxisMappings=(AxisName="MoveRight",Scale=1.000000,Key=D)
++AxisMappings=(AxisName="MoveRight",Scale=-1.000000,Key=A)
++AxisMappings=(AxisName="MoveRight",Scale=1.000000,Key=Gamepad_LeftStick_Right)
++AxisMappings=(AxisName="MoveRight",Scale=-1.000000,Key=Gamepad_LeftStick_Left)
++AxisMappings=(AxisName="MoveRight",Scale=1.000000,Key=Right)
++AxisMappings=(AxisName="MoveRight",Scale=-1.000000,Key=Left)
++AxisMappings=(AxisName="Thrust",Scale=1.000000,Key=LeftShift)
++AxisMappings=(AxisName="Thrust",Scale=1.000000,Key=RightShift)
++AxisMappings=(AxisName="Thrust",Scale=-1.000000,Key=LeftControl)
++AxisMappings=(AxisName="Thrust",Scale=-1.000000,Key=RightControl)
++AxisMappings=(AxisName="MoveUp",Scale=1.000000,Key=Gamepad_LeftY)
++AxisMappings=(AxisName="MoveRight",Scale=1.000000,Key=Gamepad_LeftX)
++AxisMappings=(AxisName="Thrust",Scale=1.000000,Key=SpaceBar)
++AxisMappings=(AxisName="Thrust",Scale=-1.000000,Key=Gamepad_RightStick_Up)
++AxisMappings=(AxisName="Thrust",Scale=1.000000,Key=Gamepad_RightStick_Down)
++AxisMappings=(AxisName="Thrust",Scale=1.000000,Key=Gamepad_RightY)
++AxisMappings=(AxisName="Thrust",Scale=1.000000,Key=Vive_Left_Trigger_Axis)
++AxisMappings=(AxisName="Thrust",Scale=-1.000000,Key=Vive_Right_Trigger_Axis)
++AxisMappings=(AxisName="Thrust",Scale=1.000000,Key=MixedReality_Left_Trigger_Axis)
++AxisMappings=(AxisName="Thrust",Scale=-1.000000,Key=MixedReality_Right_Trigger_Axis)
++AxisMappings=(AxisName="Thrust",Scale=1.000000,Key=OculusTouch_Left_Trigger_Axis)
++AxisMappings=(AxisName="Thrust",Scale=-1.000000,Key=OculusTouch_Right_Trigger_Axis)
++AxisMappings=(AxisName="Thrust",Scale=1.000000,Key=ValveIndex_Left_Trigger_Axis)
++AxisMappings=(AxisName="Thrust",Scale=-1.000000,Key=ValveIndex_Right_Trigger_Axis)
++AxisMappings=(AxisName="MoveRight",Scale=1.000000,Key=Vive_Left_Trackpad_X)
++AxisMappings=(AxisName="MoveRight",Scale=1.000000,Key=Vive_Right_Trackpad_X)
++AxisMappings=(AxisName="MoveRight",Scale=1.000000,Key=MixedReality_Left_Thumbstick_X)
++AxisMappings=(AxisName="MoveRight",Scale=1.000000,Key=MixedReality_Right_Thumbstick_X)
++AxisMappings=(AxisName="MoveRight",Scale=1.000000,Key=OculusTouch_Left_Thumbstick_X)
++AxisMappings=(AxisName="MoveRight",Scale=1.000000,Key=OculusTouch_Right_Thumbstick_X)
++AxisMappings=(AxisName="MoveRight",Scale=1.000000,Key=ValveIndex_Left_Thumbstick_X)
++AxisMappings=(AxisName="MoveRight",Scale=1.000000,Key=ValveIndex_Right_Thumbstick_X)
++AxisMappings=(AxisName="MoveUp",Scale=1.000000,Key=Vive_Left_Trackpad_Y)
++AxisMappings=(AxisName="MoveUp",Scale=1.000000,Key=Vive_Right_Trackpad_Y)
++AxisMappings=(AxisName="MoveUp",Scale=1.000000,Key=MixedReality_Left_Thumbstick_Y)
++AxisMappings=(AxisName="MoveUp",Scale=1.000000,Key=MixedReality_Right_Thumbstick_Y)
++AxisMappings=(AxisName="MoveUp",Scale=1.000000,Key=OculusTouch_Left_Thumbstick_Y)
++AxisMappings=(AxisName="MoveUp",Scale=1.000000,Key=OculusTouch_Right_Thumbstick_Y)
++AxisMappings=(AxisName="MoveUp",Scale=1.000000,Key=ValveIndex_Left_Thumbstick_Y)
++AxisMappings=(AxisName="MoveUp",Scale=1.000000,Key=ValveIndex_Right_Thumbstick_Y)
 DefaultTouchInterface=/Engine/MobileResources/HUD/DefaultVirtualJoysticks.DefaultVirtualJoysticks
-ConsoleKey=None
--ConsoleKeys=Tilde
 +ConsoleKeys=Tilde
 
 
diff --git a/Templates/TP_Puzzle/Config/DefaultInput.ini b/Templates/TP_Puzzle/Config/DefaultInput.ini
--- a/Templates/TP_Puzzle/Config/DefaultInput.ini
+++ b/Templates/TP_Puzzle/Config/DefaultInput.ini
@@ -1,11 +1,5 @@
 
 [/Script/Engine.InputSettings]
--AxisConfig=(AxisKeyName="Gamepad_LeftX",AxisProperties=(DeadZone=0.25,Exponent=1.f,Sensitivity=1.f))
--AxisConfig=(AxisKeyName="Gamepad_LeftY",AxisProperties=(DeadZone=0.25,Exponent=1.f,Sensitivity=1.f))
--AxisConfig=(AxisKeyName="Gamepad_RightX",AxisProperties=(DeadZone=0.25,Exponent=1.f,Sensitivity=1.f))
--AxisConfig=(AxisKeyName="Gamepad_RightY",AxisProperties=(DeadZone=0.25,Exponent=1.f,Sensitivity=1.f))
--AxisConfig=(AxisKeyName="MouseX",AxisProperties=(DeadZone=0.f,Exponent=1.f,Sensitivity=0.07f))
--AxisConfig=(AxisKeyName="MouseY",AxisProperties=(DeadZone=0.f,Exponent=1.f,Sensitivity=0.07f))
 +AxisConfig=(AxisKeyName="Gamepad_LeftX",AxisProperties=(DeadZone=0.250000,Sensitivity=1.000000,Exponent=1.000000,bInvert=False))
 +AxisConfig=(AxisKeyName="Gamepad_LeftY",AxisProperties=(DeadZone=0.250000,Sensitivity=1.000000,Exponent=1.000000,bInvert=False))
 +AxisConfig=(AxisKeyName="Gamepad_RightX",AxisProperties=(DeadZone=0.250000,Sensitivity=1.000000,Exponent=1.000000,bInvert=False))
@@ -27,24 +21,116 @@
 +AxisConfig=(AxisKeyName="MotionController_Right_Grip2Axis",AxisProperties=(DeadZone=0.000000,Sensitivity=1.000000,Exponent=1.000000,bInvert=False))
 +AxisConfig=(AxisKeyName="Gamepad_Special_Left_X",AxisProperties=(DeadZone=0.000000,Sensitivity=1.000000,Exponent=1.000000,bInvert=False))
 +AxisConfig=(AxisKeyName="Gamepad_Special_Left_Y",AxisProperties=(DeadZone=0.000000,Sensitivity=1.000000,Exponent=1.000000,bInvert=False))
++AxisConfig=(AxisKeyName="Daydream_Left_Trackpad_X",AxisProperties=(DeadZone=0.000000,Sensitivity=1.000000,Exponent=1.000000,bInvert=False))
++AxisConfig=(AxisKeyName="Daydream_Left_Trackpad_Y",AxisProperties=(DeadZone=0.000000,Sensitivity=1.000000,Exponent=1.000000,bInvert=False))
++AxisConfig=(AxisKeyName="Daydream_Right_Trackpad_X",AxisProperties=(DeadZone=0.000000,Sensitivity=1.000000,Exponent=1.000000,bInvert=False))
++AxisConfig=(AxisKeyName="Daydream_Right_Trackpad_Y",AxisProperties=(DeadZone=0.000000,Sensitivity=1.000000,Exponent=1.000000,bInvert=False))
++AxisConfig=(AxisKeyName="Vive_Left_Trigger_Axis",AxisProperties=(DeadZone=0.000000,Sensitivity=1.000000,Exponent=1.000000,bInvert=False))
++AxisConfig=(AxisKeyName="Vive_Left_Trackpad_X",AxisProperties=(DeadZone=0.000000,Sensitivity=1.000000,Exponent=1.000000,bInvert=False))
++AxisConfig=(AxisKeyName="Vive_Left_Trackpad_Y",AxisProperties=(DeadZone=0.000000,Sensitivity=1.000000,Exponent=1.000000,bInvert=False))
++AxisConfig=(AxisKeyName="Vive_Right_Trigger_Axis",AxisProperties=(DeadZone=0.000000,Sensitivity=1.000000,Exponent=1.000000,bInvert=False))
++AxisConfig=(AxisKeyName="Vive_Right_Trackpad_X",AxisProperties=(DeadZone=0.000000,Sensitivity=1.000000,Exponent=1.000000,bInvert=False))
++AxisConfig=(AxisKeyName="Vive_Right_Trackpad_Y",AxisProperties=(DeadZone=0.000000,Sensitivity=1.000000,Exponent=1.000000,bInvert=False))
++AxisConfig=(AxisKeyName="MixedReality_Left_Trigger_Axis",AxisProperties=(DeadZone=0.000000,Sensitivity=1.000000,Exponent=1.000000,bInvert=False))
++AxisConfig=(AxisKeyName="MixedReality_Left_Thumbstick_X",AxisProperties=(DeadZone=0.000000,Sensitivity=1.000000,Exponent=1.000000,bInvert=False))
++AxisConfig=(AxisKeyName="MixedReality_Left_Thumbstick_Y",AxisProperties=(DeadZone=0.000000,Sensitivity=1.000000,Exponent=1.000000,bInvert=False))
++AxisConfig=(AxisKeyName="MixedReality_Left_Trackpad_X",AxisProperties=(DeadZone=0.000000,Sensitivity=1.000000,Exponent=1.000000,bInvert=False))
++AxisConfig=(AxisKeyName="MixedReality_Left_Trackpad_Y",AxisProperties=(DeadZone=0.000000,Sensitivity=1.000000,Exponent=1.000000,bInvert=False))
++AxisConfig=(AxisKeyName="MixedReality_Right_Trigger_Axis",AxisProperties=(DeadZone=0.000000,Sensitivity=1.000000,Exponent=1.000000,bInvert=False))
++AxisConfig=(AxisKeyName="MixedReality_Right_Thumbstick_X",AxisProperties=(DeadZone=0.000000,Sensitivity=1.000000,Exponent=1.000000,bInvert=False))
++AxisConfig=(AxisKeyName="MixedReality_Right_Thumbstick_Y",AxisProperties=(DeadZone=0.000000,Sensitivity=1.000000,Exponent=1.000000,bInvert=False))
++AxisConfig=(AxisKeyName="MixedReality_Right_Trackpad_X",AxisProperties=(DeadZone=0.000000,Sensitivity=1.000000,Exponent=1.000000,bInvert=False))
++AxisConfig=(AxisKeyName="MixedReality_Right_Trackpad_Y",AxisProperties=(DeadZone=0.000000,Sensitivity=1.000000,Exponent=1.000000,bInvert=False))
++AxisConfig=(AxisKeyName="OculusGo_Left_Trackpad_X",AxisProperties=(DeadZone=0.000000,Sensitivity=1.000000,Exponent=1.000000,bInvert=False))
++AxisConfig=(AxisKeyName="OculusGo_Left_Trackpad_Y",AxisProperties=(DeadZone=0.000000,Sensitivity=1.000000,Exponent=1.000000,bInvert=False))
++AxisConfig=(AxisKeyName="OculusGo_Right_Trackpad_X",AxisProperties=(DeadZone=0.000000,Sensitivity=1.000000,Exponent=1.000000,bInvert=False))
++AxisConfig=(AxisKeyName="OculusGo_Right_Trackpad_Y",AxisProperties=(DeadZone=0.000000,Sensitivity=1.000000,Exponent=1.000000,bInvert=False))
++AxisConfig=(AxisKeyName="OculusTouch_Left_Grip_Axis",AxisProperties=(DeadZone=0.000000,Sensitivity=1.000000,Exponent=1.000000,bInvert=False))
++AxisConfig=(AxisKeyName="OculusTouch_Left_Trigger_Axis",AxisProperties=(DeadZone=0.000000,Sensitivity=1.000000,Exponent=1.000000,bInvert=False))
++AxisConfig=(AxisKeyName="OculusTouch_Left_Thumbstick_X",AxisProperties=(DeadZone=0.000000,Sensitivity=1.000000,Exponent=1.000000,bInvert=False))
++AxisConfig=(AxisKeyName="OculusTouch_Left_Thumbstick_Y",AxisProperties=(DeadZone=0.000000,Sensitivity=1.000000,Exponent=1.000000,bInvert=False))
++AxisConfig=(AxisKeyName="OculusTouch_Right_Grip_Axis",AxisProperties=(DeadZone=0.000000,Sensitivity=1.000000,Exponent=1.000000,bInvert=False))
++AxisConfig=(AxisKeyName="OculusTouch_Right_Trigger_Axis",AxisProperties=(DeadZone=0.000000,Sensitivity=1.000000,Exponent=1.000000,bInvert=False))
++AxisConfig=(AxisKeyName="OculusTouch_Right_Thumbstick_X",AxisProperties=(DeadZone=0.000000,Sensitivity=1.000000,Exponent=1.000000,bInvert=False))
++AxisConfig=(AxisKeyName="OculusTouch_Right_Thumbstick_Y",AxisProperties=(DeadZone=0.000000,Sensitivity=1.000000,Exponent=1.000000,bInvert=False))
++AxisConfig=(AxisKeyName="ValveIndex_Left_Grip_Axis",AxisProperties=(DeadZone=0.000000,Sensitivity=1.000000,Exponent=1.000000,bInvert=False))
++AxisConfig=(AxisKeyName="ValveIndex_Left_Grip_Force",AxisProperties=(DeadZone=0.000000,Sensitivity=1.000000,Exponent=1.000000,bInvert=False))
++AxisConfig=(AxisKeyName="ValveIndex_Left_Trigger_Axis",AxisProperties=(DeadZone=0.000000,Sensitivity=1.000000,Exponent=1.000000,bInvert=False))
++AxisConfig=(AxisKeyName="ValveIndex_Left_Thumbstick_X",AxisProperties=(DeadZone=0.000000,Sensitivity=1.000000,Exponent=1.000000,bInvert=False))
++AxisConfig=(AxisKeyName="ValveIndex_Left_Thumbstick_Y",AxisProperties=(DeadZone=0.000000,Sensitivity=1.000000,Exponent=1.000000,bInvert=False))
++AxisConfig=(AxisKeyName="ValveIndex_Left_Trackpad_X",AxisProperties=(DeadZone=0.000000,Sensitivity=1.000000,Exponent=1.000000,bInvert=False))
++AxisConfig=(AxisKeyName="ValveIndex_Left_Trackpad_Y",AxisProperties=(DeadZone=0.000000,Sensitivity=1.000000,Exponent=1.000000,bInvert=False))
++AxisConfig=(AxisKeyName="ValveIndex_Left_Trackpad_Force",AxisProperties=(DeadZone=0.000000,Sensitivity=1.000000,Exponent=1.000000,bInvert=False))
++AxisConfig=(AxisKeyName="ValveIndex_Left_Trackpad_Touch",AxisProperties=(DeadZone=0.000000,Sensitivity=1.000000,Exponent=1.000000,bInvert=False))
++AxisConfig=(AxisKeyName="ValveIndex_Right_Grip_Axis",AxisProperties=(DeadZone=0.000000,Sensitivity=1.000000,Exponent=1.000000,bInvert=False))
++AxisConfig=(AxisKeyName="ValveIndex_Right_Grip_Force",AxisProperties=(DeadZone=0.000000,Sensitivity=1.000000,Exponent=1.000000,bInvert=False))
++AxisConfig=(AxisKeyName="ValveIndex_Right_Trigger_Axis",AxisProperties=(DeadZone=0.000000,Sensitivity=1.000000,Exponent=1.000000,bInvert=False))
++AxisConfig=(AxisKeyName="ValveIndex_Right_Thumbstick_X",AxisProperties=(DeadZone=0.000000,Sensitivity=1.000000,Exponent=1.000000,bInvert=False))
++AxisConfig=(AxisKeyName="ValveIndex_Right_Thumbstick_Y",AxisProperties=(DeadZone=0.000000,Sensitivity=1.000000,Exponent=1.000000,bInvert=False))
++AxisConfig=(AxisKeyName="ValveIndex_Right_Trackpad_X",AxisProperties=(DeadZone=0.000000,Sensitivity=1.000000,Exponent=1.000000,bInvert=False))
++AxisConfig=(AxisKeyName="ValveIndex_Right_Trackpad_Y",AxisProperties=(DeadZone=0.000000,Sensitivity=1.000000,Exponent=1.000000,bInvert=False))
++AxisConfig=(AxisKeyName="ValveIndex_Right_Trackpad_Force",AxisProperties=(DeadZone=0.000000,Sensitivity=1.000000,Exponent=1.000000,bInvert=False))
++AxisConfig=(AxisKeyName="MotionController_Left_Thumbstick_Z",AxisProperties=(DeadZone=0.000000,Sensitivity=1.000000,Exponent=1.000000,bInvert=False))
++AxisConfig=(AxisKeyName="MagicLeap_Left_Trigger_Axis",AxisProperties=(DeadZone=0.000000,Sensitivity=1.000000,Exponent=1.000000,bInvert=False))
++AxisConfig=(AxisKeyName="MagicLeap_Left_Trackpad_X",AxisProperties=(DeadZone=0.000000,Sensitivity=1.000000,Exponent=1.000000,bInvert=False))
++AxisConfig=(AxisKeyName="MagicLeap_Left_Trackpad_Y",AxisProperties=(DeadZone=0.000000,Sensitivity=1.000000,Exponent=1.000000,bInvert=False))
++AxisConfig=(AxisKeyName="MagicLeap_Left_Trackpad_Force",AxisProperties=(DeadZone=0.000000,Sensitivity=1.000000,Exponent=1.000000,bInvert=False))
++AxisConfig=(AxisKeyName="MagicLeap_Left_Touch1_X",AxisProperties=(DeadZone=0.000000,Sensitivity=1.000000,Exponent=1.000000,bInvert=False))
++AxisConfig=(AxisKeyName="MagicLeap_Left_Touch1_Y",AxisProperties=(DeadZone=0.000000,Sensitivity=1.000000,Exponent=1.000000,bInvert=False))
++AxisConfig=(AxisKeyName="MagicLeap_Left_Touch1_Force",AxisProperties=(DeadZone=0.000000,Sensitivity=1.000000,Exponent=1.000000,bInvert=False))
++AxisConfig=(AxisKeyName="MotionController_Right_Thumbstick_Z",AxisProperties=(DeadZone=0.000000,Sensitivity=1.000000,Exponent=1.000000,bInvert=False))
++AxisConfig=(AxisKeyName="MagicLeap_Right_Trigger_Axis",AxisProperties=(DeadZone=0.000000,Sensitivity=1.000000,Exponent=1.000000,bInvert=False))
++AxisConfig=(AxisKeyName="MagicLeap_Right_Trackpad_X",AxisProperties=(DeadZone=0.000000,Sensitivity=1.000000,Exponent=1.000000,bInvert=False))
++AxisConfig=(AxisKeyName="MagicLeap_Right_Trackpad_Y",AxisProperties=(DeadZone=0.000000,Sensitivity=1.000000,Exponent=1.000000,bInvert=False))
++AxisConfig=(AxisKeyName="MagicLeap_Right_Trackpad_Force",AxisProperties=(DeadZone=0.000000,Sensitivity=1.000000,Exponent=1.000000,bInvert=False))
++AxisConfig=(AxisKeyName="MagicLeap_Right_Touch1_X",AxisProperties=(DeadZone=0.000000,Sensitivity=1.000000,Exponent=1.000000,bInvert=False))
++AxisConfig=(AxisKeyName="MagicLeap_Right_Touch1_Y",AxisProperties=(DeadZone=0.000000,Sensitivity=1.000000,Exponent=1.000000,bInvert=False))
++AxisConfig=(AxisKeyName="MagicLeap_Right_Touch1_Force",AxisProperties=(DeadZone=0.000000,Sensitivity=1.000000,Exponent=1.000000,bInvert=False))
++AxisConfig=(AxisKeyName="OculusTouch_Left_Thumbstick",AxisProperties=(DeadZone=0.000000,Sensitivity=1.000000,Exponent=1.000000,bInvert=False))
++AxisConfig=(AxisKeyName="OculusTouch_Left_FaceButton1",AxisProperties=(DeadZone=0.000000,Sensitivity=1.000000,Exponent=1.000000,bInvert=False))
++AxisConfig=(AxisKeyName="OculusTouch_Left_Trigger",AxisProperties=(DeadZone=0.000000,Sensitivity=1.000000,Exponent=1.000000,bInvert=False))
++AxisConfig=(AxisKeyName="OculusTouch_Left_FaceButton2",AxisProperties=(DeadZone=0.000000,Sensitivity=1.000000,Exponent=1.000000,bInvert=False))
++AxisConfig=(AxisKeyName="OculusTouch_Left_IndexPointing",AxisProperties=(DeadZone=0.000000,Sensitivity=1.000000,Exponent=1.000000,bInvert=False))
++AxisConfig=(AxisKeyName="OculusTouch_Left_ThumbUp",AxisProperties=(DeadZone=0.000000,Sensitivity=1.000000,Exponent=1.000000,bInvert=False))
++AxisConfig=(AxisKeyName="OculusTouch_Right_Thumbstick",AxisProperties=(DeadZone=0.000000,Sensitivity=1.000000,Exponent=1.000000,bInvert=False))
++AxisConfig=(AxisKeyName="OculusTouch_Right_FaceButton1",AxisProperties=(DeadZone=0.000000,Sensitivity=1.000000,Exponent=1.000000,bInvert=False))
++AxisConfig=(AxisKeyName="OculusTouch_Right_Trigger",AxisProperties=(DeadZone=0.000000,Sensitivity=1.000000,Exponent=1.000000,bInvert=False))
++AxisConfig=(AxisKeyName="OculusTouch_Right_FaceButton2",AxisProperties=(DeadZone=0.000000,Sensitivity=1.000000,Exponent=1.000000,bInvert=False))
++AxisConfig=(AxisKeyName="OculusTouch_Right_IndexPointing",AxisProperties=(DeadZone=0.000000,Sensitivity=1.000000,Exponent=1.000000,bInvert=False))
++AxisConfig=(AxisKeyName="OculusTouch_Right_ThumbUp",AxisProperties=(DeadZone=0.000000,Sensitivity=1.000000,Exponent=1.000000,bInvert=False))
++AxisConfig=(AxisKeyName="OculusTouchpad_Touchpad_X",AxisProperties=(DeadZone=0.000000,Sensitivity=1.000000,Exponent=1.000000,bInvert=False))
++AxisConfig=(AxisKeyName="OculusTouchpad_Touchpad_Y",AxisProperties=(DeadZone=0.000000,Sensitivity=1.000000,Exponent=1.000000,bInvert=False))
 bAltEnterTogglesFullscreen=True
+bF11TogglesFullscreen=True
 bUseMouseForTouch=False
 bEnableMouseSmoothing=True
 bEnableFOVScaling=True
-FOVScale=0.011110
-DoubleClickTime=0.200000
 bCaptureMouseOnLaunch=True
-DefaultViewportMouseCaptureMode=CapturePermanently_IncludingInitialMouseDown
-bDefaultViewportMouseLock=True
-+ActionMappings=(ActionName="TriggerClick",Key=MotionController_Right_Trigger,bShift=False,bCtrl=False,bAlt=False,bCmd=False)
-+ActionMappings=(ActionName="ResetVR",Key=R,bShift=False,bCtrl=False,bAlt=False,bCmd=False)
-+ActionMappings=(ActionName="ResetVR",Key=MotionController_Left_Grip1,bShift=False,bCtrl=False,bAlt=False,bCmd=False)
-+ActionMappings=(ActionName="TriggerClick",Key=Gamepad_RightTrigger,bShift=False,bCtrl=False,bAlt=False,bCmd=False)
 bAlwaysShowTouchInterface=False
 bShowConsoleOnFourFingerTap=True
+bEnableGestureRecognizer=False
+bUseAutocorrect=False
+DefaultViewportMouseCaptureMode=CapturePermanently_IncludingInitialMouseDown
+DefaultViewportMouseLockMode=LockOnCapture
+FOVScale=0.011110
+DoubleClickTime=0.200000
++ActionMappings=(ActionName="ResetVR",bShift=False,bCtrl=False,bAlt=False,bCmd=False,Key=R)
++ActionMappings=(ActionName="TriggerClick",bShift=False,bCtrl=False,bAlt=False,bCmd=False,Key=Gamepad_RightTrigger)
++ActionMappings=(ActionName="TriggerClick",bShift=False,bCtrl=False,bAlt=False,bCmd=False,Key=Daydream_Left_Select_Click)
++ActionMappings=(ActionName="TriggerClick",bShift=False,bCtrl=False,bAlt=False,bCmd=False,Key=Vive_Right_Trigger_Click)
++ActionMappings=(ActionName="TriggerClick",bShift=False,bCtrl=False,bAlt=False,bCmd=False,Key=MixedReality_Right_Trigger_Click)
++ActionMappings=(ActionName="TriggerClick",bShift=False,bCtrl=False,bAlt=False,bCmd=False,Key=OculusGo_Left_Trigger_Click)
++ActionMappings=(ActionName="TriggerClick",bShift=False,bCtrl=False,bAlt=False,bCmd=False,Key=OculusTouch_Right_Trigger_Click)
++ActionMappings=(ActionName="TriggerClick",bShift=False,bCtrl=False,bAlt=False,bCmd=False,Key=ValveIndex_Right_Trigger_Click)
++ActionMappings=(ActionName="TriggerClick",bShift=False,bCtrl=False,bAlt=False,bCmd=False,Key=MagicLeap_Left_Trigger)
++ActionMappings=(ActionName="ResetVR",bShift=False,bCtrl=False,bAlt=False,bCmd=False,Key=Daydream_Left_Trackpad_Click)
++ActionMappings=(ActionName="ResetVR",bShift=False,bCtrl=False,bAlt=False,bCmd=False,Key=Vive_Left_Grip_Click)
++ActionMappings=(ActionName="ResetVR",bShift=False,bCtrl=False,bAlt=False,bCmd=False,Key=MixedReality_Left_Thumbstick_Click)
++ActionMappings=(ActionName="ResetVR",bShift=False,bCtrl=False,bAlt=False,bCmd=False,Key=OculusGo_Left_Back_Click)
++ActionMappings=(ActionName="ResetVR",bShift=False,bCtrl=False,bAlt=False,bCmd=False,Key=OculusTouch_Left_Thumbstick_Click)
++ActionMappings=(ActionName="ResetVR",bShift=False,bCtrl=False,bAlt=False,bCmd=False,Key=ValveIndex_Left_Thumbstick_Click)
++ActionMappings=(ActionName="ResetVR",bShift=False,bCtrl=False,bAlt=False,bCmd=False,Key=MagicLeap_Left_Bumper)
 DefaultTouchInterface=None
-ConsoleKey=None
--ConsoleKeys=Tilde
 +ConsoleKeys=Tilde
 
 
diff --git a/Templates/TP_PuzzleBP/Config/DefaultInput.ini b/Templates/TP_PuzzleBP/Config/DefaultInput.ini
--- a/Templates/TP_PuzzleBP/Config/DefaultInput.ini
+++ b/Templates/TP_PuzzleBP/Config/DefaultInput.ini
@@ -1,11 +1,5 @@
 
 [/Script/Engine.InputSettings]
--AxisConfig=(AxisKeyName="Gamepad_LeftX",AxisProperties=(DeadZone=0.25,Exponent=1.f,Sensitivity=1.f))
--AxisConfig=(AxisKeyName="Gamepad_LeftY",AxisProperties=(DeadZone=0.25,Exponent=1.f,Sensitivity=1.f))
--AxisConfig=(AxisKeyName="Gamepad_RightX",AxisProperties=(DeadZone=0.25,Exponent=1.f,Sensitivity=1.f))
--AxisConfig=(AxisKeyName="Gamepad_RightY",AxisProperties=(DeadZone=0.25,Exponent=1.f,Sensitivity=1.f))
--AxisConfig=(AxisKeyName="MouseX",AxisProperties=(DeadZone=0.f,Exponent=1.f,Sensitivity=0.07f))
--AxisConfig=(AxisKeyName="MouseY",AxisProperties=(DeadZone=0.f,Exponent=1.f,Sensitivity=0.07f))
 +AxisConfig=(AxisKeyName="Gamepad_LeftX",AxisProperties=(DeadZone=0.250000,Sensitivity=1.000000,Exponent=1.000000,bInvert=False))
 +AxisConfig=(AxisKeyName="Gamepad_LeftY",AxisProperties=(DeadZone=0.250000,Sensitivity=1.000000,Exponent=1.000000,bInvert=False))
 +AxisConfig=(AxisKeyName="Gamepad_RightX",AxisProperties=(DeadZone=0.250000,Sensitivity=1.000000,Exponent=1.000000,bInvert=False))
@@ -27,24 +21,116 @@
 +AxisConfig=(AxisKeyName="MotionController_Right_Grip2Axis",AxisProperties=(DeadZone=0.000000,Sensitivity=1.000000,Exponent=1.000000,bInvert=False))
 +AxisConfig=(AxisKeyName="Gamepad_Special_Left_X",AxisProperties=(DeadZone=0.000000,Sensitivity=1.000000,Exponent=1.000000,bInvert=False))
 +AxisConfig=(AxisKeyName="Gamepad_Special_Left_Y",AxisProperties=(DeadZone=0.000000,Sensitivity=1.000000,Exponent=1.000000,bInvert=False))
++AxisConfig=(AxisKeyName="Daydream_Left_Trackpad_X",AxisProperties=(DeadZone=0.000000,Sensitivity=1.000000,Exponent=1.000000,bInvert=False))
++AxisConfig=(AxisKeyName="Daydream_Left_Trackpad_Y",AxisProperties=(DeadZone=0.000000,Sensitivity=1.000000,Exponent=1.000000,bInvert=False))
++AxisConfig=(AxisKeyName="Daydream_Right_Trackpad_X",AxisProperties=(DeadZone=0.000000,Sensitivity=1.000000,Exponent=1.000000,bInvert=False))
++AxisConfig=(AxisKeyName="Daydream_Right_Trackpad_Y",AxisProperties=(DeadZone=0.000000,Sensitivity=1.000000,Exponent=1.000000,bInvert=False))
++AxisConfig=(AxisKeyName="Vive_Left_Trigger_Axis",AxisProperties=(DeadZone=0.000000,Sensitivity=1.000000,Exponent=1.000000,bInvert=False))
++AxisConfig=(AxisKeyName="Vive_Left_Trackpad_X",AxisProperties=(DeadZone=0.000000,Sensitivity=1.000000,Exponent=1.000000,bInvert=False))
++AxisConfig=(AxisKeyName="Vive_Left_Trackpad_Y",AxisProperties=(DeadZone=0.000000,Sensitivity=1.000000,Exponent=1.000000,bInvert=False))
++AxisConfig=(AxisKeyName="Vive_Right_Trigger_Axis",AxisProperties=(DeadZone=0.000000,Sensitivity=1.000000,Exponent=1.000000,bInvert=False))
++AxisConfig=(AxisKeyName="Vive_Right_Trackpad_X",AxisProperties=(DeadZone=0.000000,Sensitivity=1.000000,Exponent=1.000000,bInvert=False))
++AxisConfig=(AxisKeyName="Vive_Right_Trackpad_Y",AxisProperties=(DeadZone=0.000000,Sensitivity=1.000000,Exponent=1.000000,bInvert=False))
++AxisConfig=(AxisKeyName="MixedReality_Left_Trigger_Axis",AxisProperties=(DeadZone=0.000000,Sensitivity=1.000000,Exponent=1.000000,bInvert=False))
++AxisConfig=(AxisKeyName="MixedReality_Left_Thumbstick_X",AxisProperties=(DeadZone=0.000000,Sensitivity=1.000000,Exponent=1.000000,bInvert=False))
++AxisConfig=(AxisKeyName="MixedReality_Left_Thumbstick_Y",AxisProperties=(DeadZone=0.000000,Sensitivity=1.000000,Exponent=1.000000,bInvert=False))
++AxisConfig=(AxisKeyName="MixedReality_Left_Trackpad_X",AxisProperties=(DeadZone=0.000000,Sensitivity=1.000000,Exponent=1.000000,bInvert=False))
++AxisConfig=(AxisKeyName="MixedReality_Left_Trackpad_Y",AxisProperties=(DeadZone=0.000000,Sensitivity=1.000000,Exponent=1.000000,bInvert=False))
++AxisConfig=(AxisKeyName="MixedReality_Right_Trigger_Axis",AxisProperties=(DeadZone=0.000000,Sensitivity=1.000000,Exponent=1.000000,bInvert=False))
++AxisConfig=(AxisKeyName="MixedReality_Right_Thumbstick_X",AxisProperties=(DeadZone=0.000000,Sensitivity=1.000000,Exponent=1.000000,bInvert=False))
++AxisConfig=(AxisKeyName="MixedReality_Right_Thumbstick_Y",AxisProperties=(DeadZone=0.000000,Sensitivity=1.000000,Exponent=1.000000,bInvert=False))
++AxisConfig=(AxisKeyName="MixedReality_Right_Trackpad_X",AxisProperties=(DeadZone=0.000000,Sensitivity=1.000000,Exponent=1.000000,bInvert=False))
++AxisConfig=(AxisKeyName="MixedReality_Right_Trackpad_Y",AxisProperties=(DeadZone=0.000000,Sensitivity=1.000000,Exponent=1.000000,bInvert=False))
++AxisConfig=(AxisKeyName="OculusGo_Left_Trackpad_X",AxisProperties=(DeadZone=0.000000,Sensitivity=1.000000,Exponent=1.000000,bInvert=False))
++AxisConfig=(AxisKeyName="OculusGo_Left_Trackpad_Y",AxisProperties=(DeadZone=0.000000,Sensitivity=1.000000,Exponent=1.000000,bInvert=False))
++AxisConfig=(AxisKeyName="OculusGo_Right_Trackpad_X",AxisProperties=(DeadZone=0.000000,Sensitivity=1.000000,Exponent=1.000000,bInvert=False))
++AxisConfig=(AxisKeyName="OculusGo_Right_Trackpad_Y",AxisProperties=(DeadZone=0.000000,Sensitivity=1.000000,Exponent=1.000000,bInvert=False))
++AxisConfig=(AxisKeyName="OculusTouch_Left_Grip_Axis",AxisProperties=(DeadZone=0.000000,Sensitivity=1.000000,Exponent=1.000000,bInvert=False))
++AxisConfig=(AxisKeyName="OculusTouch_Left_Trigger_Axis",AxisProperties=(DeadZone=0.000000,Sensitivity=1.000000,Exponent=1.000000,bInvert=False))
++AxisConfig=(AxisKeyName="OculusTouch_Left_Thumbstick_X",AxisProperties=(DeadZone=0.000000,Sensitivity=1.000000,Exponent=1.000000,bInvert=False))
++AxisConfig=(AxisKeyName="OculusTouch_Left_Thumbstick_Y",AxisProperties=(DeadZone=0.000000,Sensitivity=1.000000,Exponent=1.000000,bInvert=False))
++AxisConfig=(AxisKeyName="OculusTouch_Right_Grip_Axis",AxisProperties=(DeadZone=0.000000,Sensitivity=1.000000,Exponent=1.000000,bInvert=False))
++AxisConfig=(AxisKeyName="OculusTouch_Right_Trigger_Axis",AxisProperties=(DeadZone=0.000000,Sensitivity=1.000000,Exponent=1.000000,bInvert=False))
++AxisConfig=(AxisKeyName="OculusTouch_Right_Thumbstick_X",AxisProperties=(DeadZone=0.000000,Sensitivity=1.000000,Exponent=1.000000,bInvert=False))
++AxisConfig=(AxisKeyName="OculusTouch_Right_Thumbstick_Y",AxisProperties=(DeadZone=0.000000,Sensitivity=1.000000,Exponent=1.000000,bInvert=False))
++AxisConfig=(AxisKeyName="ValveIndex_Left_Grip_Axis",AxisProperties=(DeadZone=0.000000,Sensitivity=1.000000,Exponent=1.000000,bInvert=False))
++AxisConfig=(AxisKeyName="ValveIndex_Left_Grip_Force",AxisProperties=(DeadZone=0.000000,Sensitivity=1.000000,Exponent=1.000000,bInvert=False))
++AxisConfig=(AxisKeyName="ValveIndex_Left_Trigger_Axis",AxisProperties=(DeadZone=0.000000,Sensitivity=1.000000,Exponent=1.000000,bInvert=False))
++AxisConfig=(AxisKeyName="ValveIndex_Left_Thumbstick_X",AxisProperties=(DeadZone=0.000000,Sensitivity=1.000000,Exponent=1.000000,bInvert=False))
++AxisConfig=(AxisKeyName="ValveIndex_Left_Thumbstick_Y",AxisProperties=(DeadZone=0.000000,Sensitivity=1.000000,Exponent=1.000000,bInvert=False))
++AxisConfig=(AxisKeyName="ValveIndex_Left_Trackpad_X",AxisProperties=(DeadZone=0.000000,Sensitivity=1.000000,Exponent=1.000000,bInvert=False))
++AxisConfig=(AxisKeyName="ValveIndex_Left_Trackpad_Y",AxisProperties=(DeadZone=0.000000,Sensitivity=1.000000,Exponent=1.000000,bInvert=False))
++AxisConfig=(AxisKeyName="ValveIndex_Left_Trackpad_Force",AxisProperties=(DeadZone=0.000000,Sensitivity=1.000000,Exponent=1.000000,bInvert=False))
++AxisConfig=(AxisKeyName="ValveIndex_Left_Trackpad_Touch",AxisProperties=(DeadZone=0.000000,Sensitivity=1.000000,Exponent=1.000000,bInvert=False))
++AxisConfig=(AxisKeyName="ValveIndex_Right_Grip_Axis",AxisProperties=(DeadZone=0.000000,Sensitivity=1.000000,Exponent=1.000000,bInvert=False))
++AxisConfig=(AxisKeyName="ValveIndex_Right_Grip_Force",AxisProperties=(DeadZone=0.000000,Sensitivity=1.000000,Exponent=1.000000,bInvert=False))
++AxisConfig=(AxisKeyName="ValveIndex_Right_Trigger_Axis",AxisProperties=(DeadZone=0.000000,Sensitivity=1.000000,Exponent=1.000000,bInvert=False))
++AxisConfig=(AxisKeyName="ValveIndex_Right_Thumbstick_X",AxisProperties=(DeadZone=0.000000,Sensitivity=1.000000,Exponent=1.000000,bInvert=False))
++AxisConfig=(AxisKeyName="ValveIndex_Right_Thumbstick_Y",AxisProperties=(DeadZone=0.000000,Sensitivity=1.000000,Exponent=1.000000,bInvert=False))
++AxisConfig=(AxisKeyName="ValveIndex_Right_Trackpad_X",AxisProperties=(DeadZone=0.000000,Sensitivity=1.000000,Exponent=1.000000,bInvert=False))
++AxisConfig=(AxisKeyName="ValveIndex_Right_Trackpad_Y",AxisProperties=(DeadZone=0.000000,Sensitivity=1.000000,Exponent=1.000000,bInvert=False))
++AxisConfig=(AxisKeyName="ValveIndex_Right_Trackpad_Force",AxisProperties=(DeadZone=0.000000,Sensitivity=1.000000,Exponent=1.000000,bInvert=False))
++AxisConfig=(AxisKeyName="MotionController_Left_Thumbstick_Z",AxisProperties=(DeadZone=0.000000,Sensitivity=1.000000,Exponent=1.000000,bInvert=False))
++AxisConfig=(AxisKeyName="MagicLeap_Left_Trigger_Axis",AxisProperties=(DeadZone=0.000000,Sensitivity=1.000000,Exponent=1.000000,bInvert=False))
++AxisConfig=(AxisKeyName="MagicLeap_Left_Trackpad_X",AxisProperties=(DeadZone=0.000000,Sensitivity=1.000000,Exponent=1.000000,bInvert=False))
++AxisConfig=(AxisKeyName="MagicLeap_Left_Trackpad_Y",AxisProperties=(DeadZone=0.000000,Sensitivity=1.000000,Exponent=1.000000,bInvert=False))
++AxisConfig=(AxisKeyName="MagicLeap_Left_Trackpad_Force",AxisProperties=(DeadZone=0.000000,Sensitivity=1.000000,Exponent=1.000000,bInvert=False))
++AxisConfig=(AxisKeyName="MagicLeap_Left_Touch1_X",AxisProperties=(DeadZone=0.000000,Sensitivity=1.000000,Exponent=1.000000,bInvert=False))
++AxisConfig=(AxisKeyName="MagicLeap_Left_Touch1_Y",AxisProperties=(DeadZone=0.000000,Sensitivity=1.000000,Exponent=1.000000,bInvert=False))
++AxisConfig=(AxisKeyName="MagicLeap_Left_Touch1_Force",AxisProperties=(DeadZone=0.000000,Sensitivity=1.000000,Exponent=1.000000,bInvert=False))
++AxisConfig=(AxisKeyName="MotionController_Right_Thumbstick_Z",AxisProperties=(DeadZone=0.000000,Sensitivity=1.000000,Exponent=1.000000,bInvert=False))
++AxisConfig=(AxisKeyName="MagicLeap_Right_Trigger_Axis",AxisProperties=(DeadZone=0.000000,Sensitivity=1.000000,Exponent=1.000000,bInvert=False))
++AxisConfig=(AxisKeyName="MagicLeap_Right_Trackpad_X",AxisProperties=(DeadZone=0.000000,Sensitivity=1.000000,Exponent=1.000000,bInvert=False))
++AxisConfig=(AxisKeyName="MagicLeap_Right_Trackpad_Y",AxisProperties=(DeadZone=0.000000,Sensitivity=1.000000,Exponent=1.000000,bInvert=False))
++AxisConfig=(AxisKeyName="MagicLeap_Right_Trackpad_Force",AxisProperties=(DeadZone=0.000000,Sensitivity=1.000000,Exponent=1.000000,bInvert=False))
++AxisConfig=(AxisKeyName="MagicLeap_Right_Touch1_X",AxisProperties=(DeadZone=0.000000,Sensitivity=1.000000,Exponent=1.000000,bInvert=False))
++AxisConfig=(AxisKeyName="MagicLeap_Right_Touch1_Y",AxisProperties=(DeadZone=0.000000,Sensitivity=1.000000,Exponent=1.000000,bInvert=False))
++AxisConfig=(AxisKeyName="MagicLeap_Right_Touch1_Force",AxisProperties=(DeadZone=0.000000,Sensitivity=1.000000,Exponent=1.000000,bInvert=False))
++AxisConfig=(AxisKeyName="OculusTouch_Left_Thumbstick",AxisProperties=(DeadZone=0.000000,Sensitivity=1.000000,Exponent=1.000000,bInvert=False))
++AxisConfig=(AxisKeyName="OculusTouch_Left_FaceButton1",AxisProperties=(DeadZone=0.000000,Sensitivity=1.000000,Exponent=1.000000,bInvert=False))
++AxisConfig=(AxisKeyName="OculusTouch_Left_Trigger",AxisProperties=(DeadZone=0.000000,Sensitivity=1.000000,Exponent=1.000000,bInvert=False))
++AxisConfig=(AxisKeyName="OculusTouch_Left_FaceButton2",AxisProperties=(DeadZone=0.000000,Sensitivity=1.000000,Exponent=1.000000,bInvert=False))
++AxisConfig=(AxisKeyName="OculusTouch_Left_IndexPointing",AxisProperties=(DeadZone=0.000000,Sensitivity=1.000000,Exponent=1.000000,bInvert=False))
++AxisConfig=(AxisKeyName="OculusTouch_Left_ThumbUp",AxisProperties=(DeadZone=0.000000,Sensitivity=1.000000,Exponent=1.000000,bInvert=False))
++AxisConfig=(AxisKeyName="OculusTouch_Right_Thumbstick",AxisProperties=(DeadZone=0.000000,Sensitivity=1.000000,Exponent=1.000000,bInvert=False))
++AxisConfig=(AxisKeyName="OculusTouch_Right_FaceButton1",AxisProperties=(DeadZone=0.000000,Sensitivity=1.000000,Exponent=1.000000,bInvert=False))
++AxisConfig=(AxisKeyName="OculusTouch_Right_Trigger",AxisProperties=(DeadZone=0.000000,Sensitivity=1.000000,Exponent=1.000000,bInvert=False))
++AxisConfig=(AxisKeyName="OculusTouch_Right_FaceButton2",AxisProperties=(DeadZone=0.000000,Sensitivity=1.000000,Exponent=1.000000,bInvert=False))
++AxisConfig=(AxisKeyName="OculusTouch_Right_IndexPointing",AxisProperties=(DeadZone=0.000000,Sensitivity=1.000000,Exponent=1.000000,bInvert=False))
++AxisConfig=(AxisKeyName="OculusTouch_Right_ThumbUp",AxisProperties=(DeadZone=0.000000,Sensitivity=1.000000,Exponent=1.000000,bInvert=False))
++AxisConfig=(AxisKeyName="OculusTouchpad_Touchpad_X",AxisProperties=(DeadZone=0.000000,Sensitivity=1.000000,Exponent=1.000000,bInvert=False))
++AxisConfig=(AxisKeyName="OculusTouchpad_Touchpad_Y",AxisProperties=(DeadZone=0.000000,Sensitivity=1.000000,Exponent=1.000000,bInvert=False))
 bAltEnterTogglesFullscreen=True
+bF11TogglesFullscreen=True
 bUseMouseForTouch=False
 bEnableMouseSmoothing=True
 bEnableFOVScaling=True
-FOVScale=0.011110
-DoubleClickTime=0.200000
 bCaptureMouseOnLaunch=True
-DefaultViewportMouseCaptureMode=CapturePermanently_IncludingInitialMouseDown
-bDefaultViewportMouseLock=True
-+ActionMappings=(ActionName="TriggerClick",Key=MotionController_Right_Trigger,bShift=False,bCtrl=False,bAlt=False,bCmd=False)
-+ActionMappings=(ActionName="ResetVR",Key=R,bShift=False,bCtrl=False,bAlt=False,bCmd=False)
-+ActionMappings=(ActionName="ResetVR",Key=MotionController_Left_Grip1,bShift=False,bCtrl=False,bAlt=False,bCmd=False)
-+ActionMappings=(ActionName="TriggerClick",Key=Gamepad_RightTrigger,bShift=False,bCtrl=False,bAlt=False,bCmd=False)
 bAlwaysShowTouchInterface=False
 bShowConsoleOnFourFingerTap=True
+bEnableGestureRecognizer=False
+bUseAutocorrect=False
+DefaultViewportMouseCaptureMode=CapturePermanently_IncludingInitialMouseDown
+DefaultViewportMouseLockMode=LockOnCapture
+FOVScale=0.011110
+DoubleClickTime=0.200000
++ActionMappings=(ActionName="ResetVR",bShift=False,bCtrl=False,bAlt=False,bCmd=False,Key=R)
++ActionMappings=(ActionName="TriggerClick",bShift=False,bCtrl=False,bAlt=False,bCmd=False,Key=Gamepad_RightTrigger)
++ActionMappings=(ActionName="TriggerClick",bShift=False,bCtrl=False,bAlt=False,bCmd=False,Key=Daydream_Left_Select_Click)
++ActionMappings=(ActionName="TriggerClick",bShift=False,bCtrl=False,bAlt=False,bCmd=False,Key=Vive_Right_Trigger_Click)
++ActionMappings=(ActionName="TriggerClick",bShift=False,bCtrl=False,bAlt=False,bCmd=False,Key=MixedReality_Right_Trigger_Click)
++ActionMappings=(ActionName="TriggerClick",bShift=False,bCtrl=False,bAlt=False,bCmd=False,Key=OculusGo_Left_Trigger_Click)
++ActionMappings=(ActionName="TriggerClick",bShift=False,bCtrl=False,bAlt=False,bCmd=False,Key=OculusTouch_Right_Trigger_Click)
++ActionMappings=(ActionName="TriggerClick",bShift=False,bCtrl=False,bAlt=False,bCmd=False,Key=ValveIndex_Right_Trigger_Click)
++ActionMappings=(ActionName="TriggerClick",bShift=False,bCtrl=False,bAlt=False,bCmd=False,Key=MagicLeap_Left_Trigger)
++ActionMappings=(ActionName="ResetVR",bShift=False,bCtrl=False,bAlt=False,bCmd=False,Key=Daydream_Left_Trackpad_Click)
++ActionMappings=(ActionName="ResetVR",bShift=False,bCtrl=False,bAlt=False,bCmd=False,Key=Vive_Left_Grip_Click)
++ActionMappings=(ActionName="ResetVR",bShift=False,bCtrl=False,bAlt=False,bCmd=False,Key=MixedReality_Left_Thumbstick_Click)
++ActionMappings=(ActionName="ResetVR",bShift=False,bCtrl=False,bAlt=False,bCmd=False,Key=OculusGo_Left_Back_Click)
++ActionMappings=(ActionName="ResetVR",bShift=False,bCtrl=False,bAlt=False,bCmd=False,Key=OculusTouch_Left_Thumbstick_Click)
++ActionMappings=(ActionName="ResetVR",bShift=False,bCtrl=False,bAlt=False,bCmd=False,Key=ValveIndex_Left_Thumbstick_Click)
++ActionMappings=(ActionName="ResetVR",bShift=False,bCtrl=False,bAlt=False,bCmd=False,Key=MagicLeap_Left_Bumper)
 DefaultTouchInterface=None
-ConsoleKey=None
--ConsoleKeys=Tilde
 +ConsoleKeys=Tilde
 
 
diff --git a/Templates/TP_ThirdPerson/Config/DefaultInput.ini b/Templates/TP_ThirdPerson/Config/DefaultInput.ini
--- a/Templates/TP_ThirdPerson/Config/DefaultInput.ini
+++ b/Templates/TP_ThirdPerson/Config/DefaultInput.ini
@@ -1,19 +1,11 @@
 
 [/Script/Engine.InputSettings]
--AxisConfig=(AxisKeyName="Gamepad_LeftX",AxisProperties=(DeadZone=0.25,Exponent=1.f,Sensitivity=1.f))
--AxisConfig=(AxisKeyName="Gamepad_LeftY",AxisProperties=(DeadZone=0.25,Exponent=1.f,Sensitivity=1.f))
--AxisConfig=(AxisKeyName="Gamepad_RightX",AxisProperties=(DeadZone=0.25,Exponent=1.f,Sensitivity=1.f))
--AxisConfig=(AxisKeyName="Gamepad_RightY",AxisProperties=(DeadZone=0.25,Exponent=1.f,Sensitivity=1.f))
--AxisConfig=(AxisKeyName="MouseX",AxisProperties=(DeadZone=0.f,Exponent=1.f,Sensitivity=0.07f))
--AxisConfig=(AxisKeyName="MouseY",AxisProperties=(DeadZone=0.f,Exponent=1.f,Sensitivity=0.07f))
 +AxisConfig=(AxisKeyName="Gamepad_LeftX",AxisProperties=(DeadZone=0.250000,Sensitivity=1.000000,Exponent=1.000000,bInvert=False))
 +AxisConfig=(AxisKeyName="Gamepad_LeftY",AxisProperties=(DeadZone=0.250000,Sensitivity=1.000000,Exponent=1.000000,bInvert=False))
 +AxisConfig=(AxisKeyName="Gamepad_RightX",AxisProperties=(DeadZone=0.250000,Sensitivity=1.000000,Exponent=1.000000,bInvert=False))
 +AxisConfig=(AxisKeyName="Gamepad_RightY",AxisProperties=(DeadZone=0.250000,Sensitivity=1.000000,Exponent=1.000000,bInvert=False))
 +AxisConfig=(AxisKeyName="MouseX",AxisProperties=(DeadZone=0.000000,Sensitivity=0.070000,Exponent=1.000000,bInvert=False))
 +AxisConfig=(AxisKeyName="MouseY",AxisProperties=(DeadZone=0.000000,Sensitivity=0.070000,Exponent=1.000000,bInvert=False))
-+AxisConfig=(AxisKeyName="MotionController_Left_Thumbstick_Z",AxisProperties=(DeadZone=0.000000,Sensitivity=1.000000,Exponent=1.000000,bInvert=False))
-+AxisConfig=(AxisKeyName="MotionController_Right_Thumbstick_Z",AxisProperties=(DeadZone=0.000000,Sensitivity=1.000000,Exponent=1.000000,bInvert=False))
 +AxisConfig=(AxisKeyName="MouseWheelAxis",AxisProperties=(DeadZone=0.000000,Sensitivity=1.000000,Exponent=1.000000,bInvert=False))
 +AxisConfig=(AxisKeyName="Gamepad_LeftTriggerAxis",AxisProperties=(DeadZone=0.000000,Sensitivity=1.000000,Exponent=1.000000,bInvert=False))
 +AxisConfig=(AxisKeyName="Gamepad_RightTriggerAxis",AxisProperties=(DeadZone=0.000000,Sensitivity=1.000000,Exponent=1.000000,bInvert=False))
@@ -29,6 +21,8 @@
 +AxisConfig=(AxisKeyName="MotionController_Right_Grip2Axis",AxisProperties=(DeadZone=0.000000,Sensitivity=1.000000,Exponent=1.000000,bInvert=False))
 +AxisConfig=(AxisKeyName="Gamepad_Special_Left_X",AxisProperties=(DeadZone=0.000000,Sensitivity=1.000000,Exponent=1.000000,bInvert=False))
 +AxisConfig=(AxisKeyName="Gamepad_Special_Left_Y",AxisProperties=(DeadZone=0.000000,Sensitivity=1.000000,Exponent=1.000000,bInvert=False))
++AxisConfig=(AxisKeyName="MotionController_Left_Thumbstick_Z",AxisProperties=(DeadZone=0.000000,Sensitivity=1.000000,Exponent=1.000000,bInvert=False))
++AxisConfig=(AxisKeyName="MotionController_Right_Thumbstick_Z",AxisProperties=(DeadZone=0.000000,Sensitivity=1.000000,Exponent=1.000000,bInvert=False))
 +AxisConfig=(AxisKeyName="OculusTouch_Left_Thumbstick",AxisProperties=(DeadZone=0.000000,Sensitivity=1.000000,Exponent=1.000000,bInvert=False))
 +AxisConfig=(AxisKeyName="OculusTouch_Left_FaceButton1",AxisProperties=(DeadZone=0.000000,Sensitivity=1.000000,Exponent=1.000000,bInvert=False))
 +AxisConfig=(AxisKeyName="OculusTouch_Left_Trigger",AxisProperties=(DeadZone=0.000000,Sensitivity=1.000000,Exponent=1.000000,bInvert=False))
@@ -53,13 +47,61 @@
 +AxisConfig=(AxisKeyName="SteamVR_Knuckles_Right_MiddleGrip",AxisProperties=(DeadZone=0.000000,Sensitivity=1.000000,Exponent=1.000000,bInvert=False))
 +AxisConfig=(AxisKeyName="SteamVR_Knuckles_Right_RingGrip",AxisProperties=(DeadZone=0.000000,Sensitivity=1.000000,Exponent=1.000000,bInvert=False))
 +AxisConfig=(AxisKeyName="SteamVR_Knuckles_Right_PinkyGrip",AxisProperties=(DeadZone=0.000000,Sensitivity=1.000000,Exponent=1.000000,bInvert=False))
++AxisConfig=(AxisKeyName="Daydream_Left_Trackpad_X",AxisProperties=(DeadZone=0.000000,Sensitivity=1.000000,Exponent=1.000000,bInvert=False))
++AxisConfig=(AxisKeyName="Daydream_Left_Trackpad_Y",AxisProperties=(DeadZone=0.000000,Sensitivity=1.000000,Exponent=1.000000,bInvert=False))
++AxisConfig=(AxisKeyName="Daydream_Right_Trackpad_X",AxisProperties=(DeadZone=0.000000,Sensitivity=1.000000,Exponent=1.000000,bInvert=False))
++AxisConfig=(AxisKeyName="Daydream_Right_Trackpad_Y",AxisProperties=(DeadZone=0.000000,Sensitivity=1.000000,Exponent=1.000000,bInvert=False))
++AxisConfig=(AxisKeyName="Vive_Left_Trigger_Axis",AxisProperties=(DeadZone=0.000000,Sensitivity=1.000000,Exponent=1.000000,bInvert=False))
++AxisConfig=(AxisKeyName="Vive_Left_Trackpad_X",AxisProperties=(DeadZone=0.000000,Sensitivity=1.000000,Exponent=1.000000,bInvert=False))
++AxisConfig=(AxisKeyName="Vive_Left_Trackpad_Y",AxisProperties=(DeadZone=0.000000,Sensitivity=1.000000,Exponent=1.000000,bInvert=False))
++AxisConfig=(AxisKeyName="Vive_Right_Trigger_Axis",AxisProperties=(DeadZone=0.000000,Sensitivity=1.000000,Exponent=1.000000,bInvert=False))
++AxisConfig=(AxisKeyName="Vive_Right_Trackpad_X",AxisProperties=(DeadZone=0.000000,Sensitivity=1.000000,Exponent=1.000000,bInvert=False))
++AxisConfig=(AxisKeyName="Vive_Right_Trackpad_Y",AxisProperties=(DeadZone=0.000000,Sensitivity=1.000000,Exponent=1.000000,bInvert=False))
++AxisConfig=(AxisKeyName="MixedReality_Left_Trigger_Axis",AxisProperties=(DeadZone=0.000000,Sensitivity=1.000000,Exponent=1.000000,bInvert=False))
++AxisConfig=(AxisKeyName="MixedReality_Left_Thumbstick_X",AxisProperties=(DeadZone=0.000000,Sensitivity=1.000000,Exponent=1.000000,bInvert=False))
++AxisConfig=(AxisKeyName="MixedReality_Left_Thumbstick_Y",AxisProperties=(DeadZone=0.000000,Sensitivity=1.000000,Exponent=1.000000,bInvert=False))
++AxisConfig=(AxisKeyName="MixedReality_Left_Trackpad_X",AxisProperties=(DeadZone=0.000000,Sensitivity=1.000000,Exponent=1.000000,bInvert=False))
++AxisConfig=(AxisKeyName="MixedReality_Left_Trackpad_Y",AxisProperties=(DeadZone=0.000000,Sensitivity=1.000000,Exponent=1.000000,bInvert=False))
++AxisConfig=(AxisKeyName="MixedReality_Right_Trigger_Axis",AxisProperties=(DeadZone=0.000000,Sensitivity=1.000000,Exponent=1.000000,bInvert=False))
++AxisConfig=(AxisKeyName="MixedReality_Right_Thumbstick_X",AxisProperties=(DeadZone=0.000000,Sensitivity=1.000000,Exponent=1.000000,bInvert=False))
++AxisConfig=(AxisKeyName="MixedReality_Right_Thumbstick_Y",AxisProperties=(DeadZone=0.000000,Sensitivity=1.000000,Exponent=1.000000,bInvert=False))
++AxisConfig=(AxisKeyName="MixedReality_Right_Trackpad_X",AxisProperties=(DeadZone=0.000000,Sensitivity=1.000000,Exponent=1.000000,bInvert=False))
++AxisConfig=(AxisKeyName="MixedReality_Right_Trackpad_Y",AxisProperties=(DeadZone=0.000000,Sensitivity=1.000000,Exponent=1.000000,bInvert=False))
++AxisConfig=(AxisKeyName="OculusGo_Left_Trackpad_X",AxisProperties=(DeadZone=0.000000,Sensitivity=1.000000,Exponent=1.000000,bInvert=False))
++AxisConfig=(AxisKeyName="OculusGo_Left_Trackpad_Y",AxisProperties=(DeadZone=0.000000,Sensitivity=1.000000,Exponent=1.000000,bInvert=False))
++AxisConfig=(AxisKeyName="OculusGo_Right_Trackpad_X",AxisProperties=(DeadZone=0.000000,Sensitivity=1.000000,Exponent=1.000000,bInvert=False))
++AxisConfig=(AxisKeyName="OculusGo_Right_Trackpad_Y",AxisProperties=(DeadZone=0.000000,Sensitivity=1.000000,Exponent=1.000000,bInvert=False))
++AxisConfig=(AxisKeyName="OculusTouch_Left_Grip_Axis",AxisProperties=(DeadZone=0.000000,Sensitivity=1.000000,Exponent=1.000000,bInvert=False))
++AxisConfig=(AxisKeyName="OculusTouch_Left_Trigger_Axis",AxisProperties=(DeadZone=0.000000,Sensitivity=1.000000,Exponent=1.000000,bInvert=False))
++AxisConfig=(AxisKeyName="OculusTouch_Left_Thumbstick_X",AxisProperties=(DeadZone=0.000000,Sensitivity=1.000000,Exponent=1.000000,bInvert=False))
++AxisConfig=(AxisKeyName="OculusTouch_Left_Thumbstick_Y",AxisProperties=(DeadZone=0.000000,Sensitivity=1.000000,Exponent=1.000000,bInvert=False))
++AxisConfig=(AxisKeyName="OculusTouch_Right_Grip_Axis",AxisProperties=(DeadZone=0.000000,Sensitivity=1.000000,Exponent=1.000000,bInvert=False))
++AxisConfig=(AxisKeyName="OculusTouch_Right_Trigger_Axis",AxisProperties=(DeadZone=0.000000,Sensitivity=1.000000,Exponent=1.000000,bInvert=False))
++AxisConfig=(AxisKeyName="OculusTouch_Right_Thumbstick_X",AxisProperties=(DeadZone=0.000000,Sensitivity=1.000000,Exponent=1.000000,bInvert=False))
++AxisConfig=(AxisKeyName="OculusTouch_Right_Thumbstick_Y",AxisProperties=(DeadZone=0.000000,Sensitivity=1.000000,Exponent=1.000000,bInvert=False))
++AxisConfig=(AxisKeyName="ValveIndex_Left_Grip_Axis",AxisProperties=(DeadZone=0.000000,Sensitivity=1.000000,Exponent=1.000000,bInvert=False))
++AxisConfig=(AxisKeyName="ValveIndex_Left_Grip_Force",AxisProperties=(DeadZone=0.000000,Sensitivity=1.000000,Exponent=1.000000,bInvert=False))
++AxisConfig=(AxisKeyName="ValveIndex_Left_Trigger_Axis",AxisProperties=(DeadZone=0.000000,Sensitivity=1.000000,Exponent=1.000000,bInvert=False))
++AxisConfig=(AxisKeyName="ValveIndex_Left_Thumbstick_X",AxisProperties=(DeadZone=0.000000,Sensitivity=1.000000,Exponent=1.000000,bInvert=False))
++AxisConfig=(AxisKeyName="ValveIndex_Left_Thumbstick_Y",AxisProperties=(DeadZone=0.000000,Sensitivity=1.000000,Exponent=1.000000,bInvert=False))
++AxisConfig=(AxisKeyName="ValveIndex_Left_Trackpad_X",AxisProperties=(DeadZone=0.000000,Sensitivity=1.000000,Exponent=1.000000,bInvert=False))
++AxisConfig=(AxisKeyName="ValveIndex_Left_Trackpad_Y",AxisProperties=(DeadZone=0.000000,Sensitivity=1.000000,Exponent=1.000000,bInvert=False))
++AxisConfig=(AxisKeyName="ValveIndex_Left_Trackpad_Force",AxisProperties=(DeadZone=0.000000,Sensitivity=1.000000,Exponent=1.000000,bInvert=False))
++AxisConfig=(AxisKeyName="ValveIndex_Left_Trackpad_Touch",AxisProperties=(DeadZone=0.000000,Sensitivity=1.000000,Exponent=1.000000,bInvert=False))
++AxisConfig=(AxisKeyName="ValveIndex_Right_Grip_Axis",AxisProperties=(DeadZone=0.000000,Sensitivity=1.000000,Exponent=1.000000,bInvert=False))
++AxisConfig=(AxisKeyName="ValveIndex_Right_Grip_Force",AxisProperties=(DeadZone=0.000000,Sensitivity=1.000000,Exponent=1.000000,bInvert=False))
++AxisConfig=(AxisKeyName="ValveIndex_Right_Trigger_Axis",AxisProperties=(DeadZone=0.000000,Sensitivity=1.000000,Exponent=1.000000,bInvert=False))
++AxisConfig=(AxisKeyName="ValveIndex_Right_Thumbstick_X",AxisProperties=(DeadZone=0.000000,Sensitivity=1.000000,Exponent=1.000000,bInvert=False))
++AxisConfig=(AxisKeyName="ValveIndex_Right_Thumbstick_Y",AxisProperties=(DeadZone=0.000000,Sensitivity=1.000000,Exponent=1.000000,bInvert=False))
++AxisConfig=(AxisKeyName="ValveIndex_Right_Trackpad_X",AxisProperties=(DeadZone=0.000000,Sensitivity=1.000000,Exponent=1.000000,bInvert=False))
++AxisConfig=(AxisKeyName="ValveIndex_Right_Trackpad_Y",AxisProperties=(DeadZone=0.000000,Sensitivity=1.000000,Exponent=1.000000,bInvert=False))
++AxisConfig=(AxisKeyName="ValveIndex_Right_Trackpad_Force",AxisProperties=(DeadZone=0.000000,Sensitivity=1.000000,Exponent=1.000000,bInvert=False))
 bAltEnterTogglesFullscreen=True
 bF11TogglesFullscreen=True
 bUseMouseForTouch=False
 bEnableMouseSmoothing=True
 bEnableFOVScaling=True
 bCaptureMouseOnLaunch=True
-bDefaultViewportMouseLock=False
 bAlwaysShowTouchInterface=False
 bShowConsoleOnFourFingerTap=True
 bEnableGestureRecognizer=False
@@ -70,9 +112,25 @@ FOVScale=0.011110
 DoubleClickTime=0.200000
 +ActionMappings=(ActionName="Jump",bShift=False,bCtrl=False,bAlt=False,bCmd=False,Key=SpaceBar)
 +ActionMappings=(ActionName="Jump",bShift=False,bCtrl=False,bAlt=False,bCmd=False,Key=Gamepad_FaceButton_Bottom)
++ActionMappings=(ActionName="Jump",bShift=False,bCtrl=False,bAlt=False,bCmd=False,Key=Daydream_Left_Select_Click)
 +ActionMappings=(ActionName="ResetVR",bShift=False,bCtrl=False,bAlt=False,bCmd=False,Key=R)
-+ActionMappings=(ActionName="ResetVR",bShift=False,bCtrl=False,bAlt=False,bCmd=False,Key=MotionController_Left_Grip1)
-+ActionMappings=(ActionName="Jump",bShift=False,bCtrl=False,bAlt=False,bCmd=False,Key=OculusTouchpad_Touchpad)
++ActionMappings=(ActionName="ResetVR",bShift=False,bCtrl=False,bAlt=False,bCmd=False,Key=Daydream_Left_Trackpad_Click)
++ActionMappings=(ActionName="Jump",bShift=False,bCtrl=False,bAlt=False,bCmd=False,Key=Vive_Left_Trigger_Click)
++ActionMappings=(ActionName="Jump",bShift=False,bCtrl=False,bAlt=False,bCmd=False,Key=Vive_Right_Trigger_Click)
++ActionMappings=(ActionName="Jump",bShift=False,bCtrl=False,bAlt=False,bCmd=False,Key=MixedReality_Left_Trigger_Click)
++ActionMappings=(ActionName="Jump",bShift=False,bCtrl=False,bAlt=False,bCmd=False,Key=MixedReality_Right_Trigger_Click)
++ActionMappings=(ActionName="Jump",bShift=False,bCtrl=False,bAlt=False,bCmd=False,Key=OculusGo_Left_Trigger_Click)
++ActionMappings=(ActionName="Jump",bShift=False,bCtrl=False,bAlt=False,bCmd=False,Key=OculusTouch_Left_Trigger_Click)
++ActionMappings=(ActionName="Jump",bShift=False,bCtrl=False,bAlt=False,bCmd=False,Key=OculusTouch_Right_Trigger_Click)
++ActionMappings=(ActionName="Jump",bShift=False,bCtrl=False,bAlt=False,bCmd=False,Key=ValveIndex_Left_Trigger_Click)
++ActionMappings=(ActionName="Jump",bShift=False,bCtrl=False,bAlt=False,bCmd=False,Key=ValveIndex_Right_Trigger_Click)
++ActionMappings=(ActionName="Jump",bShift=False,bCtrl=False,bAlt=False,bCmd=False,Key=MagicLeap_Left_Trigger)
++ActionMappings=(ActionName="ResetVR",bShift=False,bCtrl=False,bAlt=False,bCmd=False,Key=Vive_Left_Grip_Click)
++ActionMappings=(ActionName="ResetVR",bShift=False,bCtrl=False,bAlt=False,bCmd=False,Key=MixedReality_Left_Thumbstick_Click)
++ActionMappings=(ActionName="ResetVR",bShift=False,bCtrl=False,bAlt=False,bCmd=False,Key=OculusGo_Left_Trackpad_Click)
++ActionMappings=(ActionName="ResetVR",bShift=False,bCtrl=False,bAlt=False,bCmd=False,Key=OculusTouch_Left_Thumbstick_Click)
++ActionMappings=(ActionName="ResetVR",bShift=False,bCtrl=False,bAlt=False,bCmd=False,Key=ValveIndex_Left_Thumbstick_Click)
++ActionMappings=(ActionName="ResetVR",bShift=False,bCtrl=False,bAlt=False,bCmd=False,Key=MagicLeap_Left_Bumper)
 +AxisMappings=(AxisName="MoveForward",Scale=1.000000,Key=W)
 +AxisMappings=(AxisName="MoveForward",Scale=-1.000000,Key=S)
 +AxisMappings=(AxisName="MoveForward",Scale=1.000000,Key=Up)
@@ -87,9 +145,25 @@ DoubleClickTime=0.200000
 +AxisMappings=(AxisName="Turn",Scale=1.000000,Key=MouseX)
 +AxisMappings=(AxisName="LookUpRate",Scale=1.000000,Key=Gamepad_RightY)
 +AxisMappings=(AxisName="LookUp",Scale=-1.000000,Key=MouseY)
++AxisMappings=(AxisName="TurnRate",Scale=-1.000000,Key=Vive_Right_Trackpad_X)
++AxisMappings=(AxisName="MoveForward",Scale=1.000000,Key=Daydream_Left_Trackpad_Y)
++AxisMappings=(AxisName="MoveForward",Scale=1.000000,Key=Vive_Left_Trackpad_Y)
++AxisMappings=(AxisName="MoveRight",Scale=1.000000,Key=Daydream_Left_Trackpad_X)
++AxisMappings=(AxisName="MoveRight",Scale=1.000000,Key=Vive_Left_Trackpad_X)
++AxisMappings=(AxisName="MoveRight",Scale=1.000000,Key=MixedReality_Left_Thumbstick_X)
++AxisMappings=(AxisName="MoveRight",Scale=1.000000,Key=OculusGo_Left_Trackpad_X)
++AxisMappings=(AxisName="MoveForward",Scale=1.000000,Key=MixedReality_Left_Thumbstick_Y)
++AxisMappings=(AxisName="MoveForward",Scale=1.000000,Key=OculusGo_Left_Trackpad_Y)
++AxisMappings=(AxisName="TurnRate",Scale=-1.000000,Key=MixedReality_Right_Thumbstick_X)
++AxisMappings=(AxisName="TurnRate",Scale=-1.000000,Key=OculusTouch_Right_Thumbstick_X)
++AxisMappings=(AxisName="TurnRate",Scale=-1.000000,Key=ValveIndex_Right_Thumbstick_X)
++AxisMappings=(AxisName="MoveForward",Scale=1.000000,Key=OculusTouch_Left_Thumbstick_Y)
++AxisMappings=(AxisName="MoveForward",Scale=1.000000,Key=ValveIndex_Left_Thumbstick_Y)
++AxisMappings=(AxisName="MoveForward",Scale=1.000000,Key=MagicLeap_Left_Trackpad_Y)
++AxisMappings=(AxisName="MoveRight",Scale=1.000000,Key=OculusTouch_Left_Thumbstick_X)
++AxisMappings=(AxisName="MoveRight",Scale=1.000000,Key=ValveIndex_Left_Thumbstick_X)
++AxisMappings=(AxisName="MoveRight",Scale=1.000000,Key=MagicLeap_Left_Trackpad_X)
 DefaultTouchInterface=/Engine/MobileResources/HUD/DefaultVirtualJoysticks.DefaultVirtualJoysticks
-ConsoleKey=None
--ConsoleKeys=Tilde
 +ConsoleKeys=Tilde
 
 
diff --git a/Templates/TP_ThirdPersonBP/Config/DefaultInput.ini b/Templates/TP_ThirdPersonBP/Config/DefaultInput.ini
--- a/Templates/TP_ThirdPersonBP/Config/DefaultInput.ini
+++ b/Templates/TP_ThirdPersonBP/Config/DefaultInput.ini
@@ -1,11 +1,5 @@
 
 [/Script/Engine.InputSettings]
--AxisConfig=(AxisKeyName="Gamepad_LeftX",AxisProperties=(DeadZone=0.25,Exponent=1.f,Sensitivity=1.f))
--AxisConfig=(AxisKeyName="Gamepad_LeftY",AxisProperties=(DeadZone=0.25,Exponent=1.f,Sensitivity=1.f))
--AxisConfig=(AxisKeyName="Gamepad_RightX",AxisProperties=(DeadZone=0.25,Exponent=1.f,Sensitivity=1.f))
--AxisConfig=(AxisKeyName="Gamepad_RightY",AxisProperties=(DeadZone=0.25,Exponent=1.f,Sensitivity=1.f))
--AxisConfig=(AxisKeyName="MouseX",AxisProperties=(DeadZone=0.f,Exponent=1.f,Sensitivity=0.07f))
--AxisConfig=(AxisKeyName="MouseY",AxisProperties=(DeadZone=0.f,Exponent=1.f,Sensitivity=0.07f))
 +AxisConfig=(AxisKeyName="Gamepad_LeftX",AxisProperties=(DeadZone=0.250000,Sensitivity=1.000000,Exponent=1.000000,bInvert=False))
 +AxisConfig=(AxisKeyName="Gamepad_LeftY",AxisProperties=(DeadZone=0.250000,Sensitivity=1.000000,Exponent=1.000000,bInvert=False))
 +AxisConfig=(AxisKeyName="Gamepad_RightX",AxisProperties=(DeadZone=0.250000,Sensitivity=1.000000,Exponent=1.000000,bInvert=False))
@@ -53,13 +47,61 @@
 +AxisConfig=(AxisKeyName="SteamVR_Knuckles_Right_MiddleGrip",AxisProperties=(DeadZone=0.000000,Sensitivity=1.000000,Exponent=1.000000,bInvert=False))
 +AxisConfig=(AxisKeyName="SteamVR_Knuckles_Right_RingGrip",AxisProperties=(DeadZone=0.000000,Sensitivity=1.000000,Exponent=1.000000,bInvert=False))
 +AxisConfig=(AxisKeyName="SteamVR_Knuckles_Right_PinkyGrip",AxisProperties=(DeadZone=0.000000,Sensitivity=1.000000,Exponent=1.000000,bInvert=False))
++AxisConfig=(AxisKeyName="Daydream_Left_Trackpad_X",AxisProperties=(DeadZone=0.000000,Sensitivity=1.000000,Exponent=1.000000,bInvert=False))
++AxisConfig=(AxisKeyName="Daydream_Left_Trackpad_Y",AxisProperties=(DeadZone=0.000000,Sensitivity=1.000000,Exponent=1.000000,bInvert=False))
++AxisConfig=(AxisKeyName="Daydream_Right_Trackpad_X",AxisProperties=(DeadZone=0.000000,Sensitivity=1.000000,Exponent=1.000000,bInvert=False))
++AxisConfig=(AxisKeyName="Daydream_Right_Trackpad_Y",AxisProperties=(DeadZone=0.000000,Sensitivity=1.000000,Exponent=1.000000,bInvert=False))
++AxisConfig=(AxisKeyName="Vive_Left_Trigger_Axis",AxisProperties=(DeadZone=0.000000,Sensitivity=1.000000,Exponent=1.000000,bInvert=False))
++AxisConfig=(AxisKeyName="Vive_Left_Trackpad_X",AxisProperties=(DeadZone=0.000000,Sensitivity=1.000000,Exponent=1.000000,bInvert=False))
++AxisConfig=(AxisKeyName="Vive_Left_Trackpad_Y",AxisProperties=(DeadZone=0.000000,Sensitivity=1.000000,Exponent=1.000000,bInvert=False))
++AxisConfig=(AxisKeyName="Vive_Right_Trigger_Axis",AxisProperties=(DeadZone=0.000000,Sensitivity=1.000000,Exponent=1.000000,bInvert=False))
++AxisConfig=(AxisKeyName="Vive_Right_Trackpad_X",AxisProperties=(DeadZone=0.000000,Sensitivity=1.000000,Exponent=1.000000,bInvert=False))
++AxisConfig=(AxisKeyName="Vive_Right_Trackpad_Y",AxisProperties=(DeadZone=0.000000,Sensitivity=1.000000,Exponent=1.000000,bInvert=False))
++AxisConfig=(AxisKeyName="MixedReality_Left_Trigger_Axis",AxisProperties=(DeadZone=0.000000,Sensitivity=1.000000,Exponent=1.000000,bInvert=False))
++AxisConfig=(AxisKeyName="MixedReality_Left_Thumbstick_X",AxisProperties=(DeadZone=0.000000,Sensitivity=1.000000,Exponent=1.000000,bInvert=False))
++AxisConfig=(AxisKeyName="MixedReality_Left_Thumbstick_Y",AxisProperties=(DeadZone=0.000000,Sensitivity=1.000000,Exponent=1.000000,bInvert=False))
++AxisConfig=(AxisKeyName="MixedReality_Left_Trackpad_X",AxisProperties=(DeadZone=0.000000,Sensitivity=1.000000,Exponent=1.000000,bInvert=False))
++AxisConfig=(AxisKeyName="MixedReality_Left_Trackpad_Y",AxisProperties=(DeadZone=0.000000,Sensitivity=1.000000,Exponent=1.000000,bInvert=False))
++AxisConfig=(AxisKeyName="MixedReality_Right_Trigger_Axis",AxisProperties=(DeadZone=0.000000,Sensitivity=1.000000,Exponent=1.000000,bInvert=False))
++AxisConfig=(AxisKeyName="MixedReality_Right_Thumbstick_X",AxisProperties=(DeadZone=0.000000,Sensitivity=1.000000,Exponent=1.000000,bInvert=False))
++AxisConfig=(AxisKeyName="MixedReality_Right_Thumbstick_Y",AxisProperties=(DeadZone=0.000000,Sensitivity=1.000000,Exponent=1.000000,bInvert=False))
++AxisConfig=(AxisKeyName="MixedReality_Right_Trackpad_X",AxisProperties=(DeadZone=0.000000,Sensitivity=1.000000,Exponent=1.000000,bInvert=False))
++AxisConfig=(AxisKeyName="MixedReality_Right_Trackpad_Y",AxisProperties=(DeadZone=0.000000,Sensitivity=1.000000,Exponent=1.000000,bInvert=False))
++AxisConfig=(AxisKeyName="OculusGo_Left_Trackpad_X",AxisProperties=(DeadZone=0.000000,Sensitivity=1.000000,Exponent=1.000000,bInvert=False))
++AxisConfig=(AxisKeyName="OculusGo_Left_Trackpad_Y",AxisProperties=(DeadZone=0.000000,Sensitivity=1.000000,Exponent=1.000000,bInvert=False))
++AxisConfig=(AxisKeyName="OculusGo_Right_Trackpad_X",AxisProperties=(DeadZone=0.000000,Sensitivity=1.000000,Exponent=1.000000,bInvert=False))
++AxisConfig=(AxisKeyName="OculusGo_Right_Trackpad_Y",AxisProperties=(DeadZone=0.000000,Sensitivity=1.000000,Exponent=1.000000,bInvert=False))
++AxisConfig=(AxisKeyName="OculusTouch_Left_Grip_Axis",AxisProperties=(DeadZone=0.000000,Sensitivity=1.000000,Exponent=1.000000,bInvert=False))
++AxisConfig=(AxisKeyName="OculusTouch_Left_Trigger_Axis",AxisProperties=(DeadZone=0.000000,Sensitivity=1.000000,Exponent=1.000000,bInvert=False))
++AxisConfig=(AxisKeyName="OculusTouch_Left_Thumbstick_X",AxisProperties=(DeadZone=0.000000,Sensitivity=1.000000,Exponent=1.000000,bInvert=False))
++AxisConfig=(AxisKeyName="OculusTouch_Left_Thumbstick_Y",AxisProperties=(DeadZone=0.000000,Sensitivity=1.000000,Exponent=1.000000,bInvert=False))
++AxisConfig=(AxisKeyName="OculusTouch_Right_Grip_Axis",AxisProperties=(DeadZone=0.000000,Sensitivity=1.000000,Exponent=1.000000,bInvert=False))
++AxisConfig=(AxisKeyName="OculusTouch_Right_Trigger_Axis",AxisProperties=(DeadZone=0.000000,Sensitivity=1.000000,Exponent=1.000000,bInvert=False))
++AxisConfig=(AxisKeyName="OculusTouch_Right_Thumbstick_X",AxisProperties=(DeadZone=0.000000,Sensitivity=1.000000,Exponent=1.000000,bInvert=False))
++AxisConfig=(AxisKeyName="OculusTouch_Right_Thumbstick_Y",AxisProperties=(DeadZone=0.000000,Sensitivity=1.000000,Exponent=1.000000,bInvert=False))
++AxisConfig=(AxisKeyName="ValveIndex_Left_Grip_Axis",AxisProperties=(DeadZone=0.000000,Sensitivity=1.000000,Exponent=1.000000,bInvert=False))
++AxisConfig=(AxisKeyName="ValveIndex_Left_Grip_Force",AxisProperties=(DeadZone=0.000000,Sensitivity=1.000000,Exponent=1.000000,bInvert=False))
++AxisConfig=(AxisKeyName="ValveIndex_Left_Trigger_Axis",AxisProperties=(DeadZone=0.000000,Sensitivity=1.000000,Exponent=1.000000,bInvert=False))
++AxisConfig=(AxisKeyName="ValveIndex_Left_Thumbstick_X",AxisProperties=(DeadZone=0.000000,Sensitivity=1.000000,Exponent=1.000000,bInvert=False))
++AxisConfig=(AxisKeyName="ValveIndex_Left_Thumbstick_Y",AxisProperties=(DeadZone=0.000000,Sensitivity=1.000000,Exponent=1.000000,bInvert=False))
++AxisConfig=(AxisKeyName="ValveIndex_Left_Trackpad_X",AxisProperties=(DeadZone=0.000000,Sensitivity=1.000000,Exponent=1.000000,bInvert=False))
++AxisConfig=(AxisKeyName="ValveIndex_Left_Trackpad_Y",AxisProperties=(DeadZone=0.000000,Sensitivity=1.000000,Exponent=1.000000,bInvert=False))
++AxisConfig=(AxisKeyName="ValveIndex_Left_Trackpad_Force",AxisProperties=(DeadZone=0.000000,Sensitivity=1.000000,Exponent=1.000000,bInvert=False))
++AxisConfig=(AxisKeyName="ValveIndex_Left_Trackpad_Touch",AxisProperties=(DeadZone=0.000000,Sensitivity=1.000000,Exponent=1.000000,bInvert=False))
++AxisConfig=(AxisKeyName="ValveIndex_Right_Grip_Axis",AxisProperties=(DeadZone=0.000000,Sensitivity=1.000000,Exponent=1.000000,bInvert=False))
++AxisConfig=(AxisKeyName="ValveIndex_Right_Grip_Force",AxisProperties=(DeadZone=0.000000,Sensitivity=1.000000,Exponent=1.000000,bInvert=False))
++AxisConfig=(AxisKeyName="ValveIndex_Right_Trigger_Axis",AxisProperties=(DeadZone=0.000000,Sensitivity=1.000000,Exponent=1.000000,bInvert=False))
++AxisConfig=(AxisKeyName="ValveIndex_Right_Thumbstick_X",AxisProperties=(DeadZone=0.000000,Sensitivity=1.000000,Exponent=1.000000,bInvert=False))
++AxisConfig=(AxisKeyName="ValveIndex_Right_Thumbstick_Y",AxisProperties=(DeadZone=0.000000,Sensitivity=1.000000,Exponent=1.000000,bInvert=False))
++AxisConfig=(AxisKeyName="ValveIndex_Right_Trackpad_X",AxisProperties=(DeadZone=0.000000,Sensitivity=1.000000,Exponent=1.000000,bInvert=False))
++AxisConfig=(AxisKeyName="ValveIndex_Right_Trackpad_Y",AxisProperties=(DeadZone=0.000000,Sensitivity=1.000000,Exponent=1.000000,bInvert=False))
++AxisConfig=(AxisKeyName="ValveIndex_Right_Trackpad_Force",AxisProperties=(DeadZone=0.000000,Sensitivity=1.000000,Exponent=1.000000,bInvert=False))
 bAltEnterTogglesFullscreen=True
 bF11TogglesFullscreen=True
 bUseMouseForTouch=False
 bEnableMouseSmoothing=True
 bEnableFOVScaling=True
 bCaptureMouseOnLaunch=True
-bDefaultViewportMouseLock=True
 bAlwaysShowTouchInterface=False
 bShowConsoleOnFourFingerTap=True
 bEnableGestureRecognizer=False
@@ -70,11 +112,25 @@ FOVScale=0.011110
 DoubleClickTime=0.200000
 +ActionMappings=(ActionName="Jump",bShift=False,bCtrl=False,bAlt=False,bCmd=False,Key=SpaceBar)
 +ActionMappings=(ActionName="Jump",bShift=False,bCtrl=False,bAlt=False,bCmd=False,Key=Gamepad_FaceButton_Bottom)
-+ActionMappings=(ActionName="Jump",bShift=False,bCtrl=False,bAlt=False,bCmd=False,Key=MotionController_Left_Trigger)
++ActionMappings=(ActionName="Jump",bShift=False,bCtrl=False,bAlt=False,bCmd=False,Key=Daydream_Left_Select_Click)
 +ActionMappings=(ActionName="ResetVR",bShift=False,bCtrl=False,bAlt=False,bCmd=False,Key=R)
-+ActionMappings=(ActionName="ResetVR",bShift=False,bCtrl=False,bAlt=False,bCmd=False,Key=MotionController_Left_Grip1)
-+ActionMappings=(ActionName="Jump",bShift=False,bCtrl=False,bAlt=False,bCmd=False,Key=MotionController_Right_Trigger)
-+ActionMappings=(ActionName="Jump",bShift=False,bCtrl=False,bAlt=False,bCmd=False,Key=OculusTouchpad_Touchpad)
++ActionMappings=(ActionName="ResetVR",bShift=False,bCtrl=False,bAlt=False,bCmd=False,Key=Daydream_Left_Trackpad_Click)
++ActionMappings=(ActionName="Jump",bShift=False,bCtrl=False,bAlt=False,bCmd=False,Key=Vive_Left_Trigger_Click)
++ActionMappings=(ActionName="Jump",bShift=False,bCtrl=False,bAlt=False,bCmd=False,Key=Vive_Right_Trigger_Click)
++ActionMappings=(ActionName="Jump",bShift=False,bCtrl=False,bAlt=False,bCmd=False,Key=MixedReality_Left_Trigger_Click)
++ActionMappings=(ActionName="Jump",bShift=False,bCtrl=False,bAlt=False,bCmd=False,Key=MixedReality_Right_Trigger_Click)
++ActionMappings=(ActionName="Jump",bShift=False,bCtrl=False,bAlt=False,bCmd=False,Key=OculusGo_Left_Trigger_Click)
++ActionMappings=(ActionName="Jump",bShift=False,bCtrl=False,bAlt=False,bCmd=False,Key=OculusTouch_Left_Trigger_Click)
++ActionMappings=(ActionName="Jump",bShift=False,bCtrl=False,bAlt=False,bCmd=False,Key=OculusTouch_Right_Trigger_Click)
++ActionMappings=(ActionName="Jump",bShift=False,bCtrl=False,bAlt=False,bCmd=False,Key=ValveIndex_Left_Trigger_Click)
++ActionMappings=(ActionName="Jump",bShift=False,bCtrl=False,bAlt=False,bCmd=False,Key=ValveIndex_Right_Trigger_Click)
++ActionMappings=(ActionName="Jump",bShift=False,bCtrl=False,bAlt=False,bCmd=False,Key=MagicLeap_Left_Trigger)
++ActionMappings=(ActionName="ResetVR",bShift=False,bCtrl=False,bAlt=False,bCmd=False,Key=Vive_Left_Grip_Click)
++ActionMappings=(ActionName="ResetVR",bShift=False,bCtrl=False,bAlt=False,bCmd=False,Key=MixedReality_Left_Thumbstick_Click)
++ActionMappings=(ActionName="ResetVR",bShift=False,bCtrl=False,bAlt=False,bCmd=False,Key=OculusGo_Left_Trackpad_Click)
++ActionMappings=(ActionName="ResetVR",bShift=False,bCtrl=False,bAlt=False,bCmd=False,Key=OculusTouch_Left_Thumbstick_Click)
++ActionMappings=(ActionName="ResetVR",bShift=False,bCtrl=False,bAlt=False,bCmd=False,Key=ValveIndex_Left_Thumbstick_Click)
++ActionMappings=(ActionName="ResetVR",bShift=False,bCtrl=False,bAlt=False,bCmd=False,Key=MagicLeap_Left_Bumper)
 +AxisMappings=(AxisName="MoveForward",Scale=1.000000,Key=W)
 +AxisMappings=(AxisName="MoveForward",Scale=-1.000000,Key=S)
 +AxisMappings=(AxisName="MoveForward",Scale=1.000000,Key=Up)
@@ -89,12 +145,25 @@ DoubleClickTime=0.200000
 +AxisMappings=(AxisName="Turn",Scale=1.000000,Key=MouseX)
 +AxisMappings=(AxisName="LookUpRate",Scale=1.000000,Key=Gamepad_RightY)
 +AxisMappings=(AxisName="LookUp",Scale=-1.000000,Key=MouseY)
-+AxisMappings=(AxisName="MoveForward",Scale=-1.000000,Key=MotionController_Left_Thumbstick_Y)
-+AxisMappings=(AxisName="MoveRight",Scale=1.000000,Key=MotionController_Left_Thumbstick_X)
-+AxisMappings=(AxisName="TurnRate",Scale=-1.000000,Key=MotionController_Right_Thumbstick_X)
++AxisMappings=(AxisName="TurnRate",Scale=-1.000000,Key=Vive_Right_Trackpad_X)
++AxisMappings=(AxisName="MoveForward",Scale=1.000000,Key=Daydream_Left_Trackpad_Y)
++AxisMappings=(AxisName="MoveForward",Scale=1.000000,Key=Vive_Left_Trackpad_Y)
++AxisMappings=(AxisName="MoveRight",Scale=1.000000,Key=Daydream_Left_Trackpad_X)
++AxisMappings=(AxisName="MoveRight",Scale=1.000000,Key=Vive_Left_Trackpad_X)
++AxisMappings=(AxisName="MoveRight",Scale=1.000000,Key=MixedReality_Left_Thumbstick_X)
++AxisMappings=(AxisName="MoveRight",Scale=1.000000,Key=OculusGo_Left_Trackpad_X)
++AxisMappings=(AxisName="MoveForward",Scale=1.000000,Key=MixedReality_Left_Thumbstick_Y)
++AxisMappings=(AxisName="MoveForward",Scale=1.000000,Key=OculusGo_Left_Trackpad_Y)
++AxisMappings=(AxisName="TurnRate",Scale=-1.000000,Key=MixedReality_Right_Thumbstick_X)
++AxisMappings=(AxisName="TurnRate",Scale=-1.000000,Key=OculusTouch_Right_Thumbstick_X)
++AxisMappings=(AxisName="TurnRate",Scale=-1.000000,Key=ValveIndex_Right_Thumbstick_X)
++AxisMappings=(AxisName="MoveForward",Scale=1.000000,Key=OculusTouch_Left_Thumbstick_Y)
++AxisMappings=(AxisName="MoveForward",Scale=1.000000,Key=ValveIndex_Left_Thumbstick_Y)
++AxisMappings=(AxisName="MoveForward",Scale=1.000000,Key=MagicLeap_Left_Trackpad_Y)
++AxisMappings=(AxisName="MoveRight",Scale=1.000000,Key=OculusTouch_Left_Thumbstick_X)
++AxisMappings=(AxisName="MoveRight",Scale=1.000000,Key=ValveIndex_Left_Thumbstick_X)
++AxisMappings=(AxisName="MoveRight",Scale=1.000000,Key=MagicLeap_Left_Trackpad_X)
 DefaultTouchInterface=/Engine/MobileResources/HUD/DefaultVirtualJoysticks.DefaultVirtualJoysticks
-ConsoleKey=None
--ConsoleKeys=Tilde
 +ConsoleKeys=Tilde
 
 
diff --git a/Templates/TP_TopDown/Config/DefaultInput.ini b/Templates/TP_TopDown/Config/DefaultInput.ini
--- a/Templates/TP_TopDown/Config/DefaultInput.ini
+++ b/Templates/TP_TopDown/Config/DefaultInput.ini
@@ -1,10 +1,5 @@
+
 [/Script/Engine.InputSettings]
--AxisConfig=(AxisKeyName="Gamepad_LeftX",AxisProperties=(DeadZone=0.25,Exponent=1.f,Sensitivity=1.f))
--AxisConfig=(AxisKeyName="Gamepad_LeftY",AxisProperties=(DeadZone=0.25,Exponent=1.f,Sensitivity=1.f))
--AxisConfig=(AxisKeyName="Gamepad_RightX",AxisProperties=(DeadZone=0.25,Exponent=1.f,Sensitivity=1.f))
--AxisConfig=(AxisKeyName="Gamepad_RightY",AxisProperties=(DeadZone=0.25,Exponent=1.f,Sensitivity=1.f))
--AxisConfig=(AxisKeyName="MouseX",AxisProperties=(DeadZone=0.f,Exponent=1.f,Sensitivity=0.07f))
--AxisConfig=(AxisKeyName="MouseY",AxisProperties=(DeadZone=0.f,Exponent=1.f,Sensitivity=0.07f))
 +AxisConfig=(AxisKeyName="Gamepad_LeftX",AxisProperties=(DeadZone=0.250000,Sensitivity=1.000000,Exponent=1.000000,bInvert=False))
 +AxisConfig=(AxisKeyName="Gamepad_LeftY",AxisProperties=(DeadZone=0.250000,Sensitivity=1.000000,Exponent=1.000000,bInvert=False))
 +AxisConfig=(AxisKeyName="Gamepad_RightX",AxisProperties=(DeadZone=0.250000,Sensitivity=1.000000,Exponent=1.000000,bInvert=False))
@@ -26,31 +21,129 @@
 +AxisConfig=(AxisKeyName="MotionController_Right_Grip2Axis",AxisProperties=(DeadZone=0.000000,Sensitivity=1.000000,Exponent=1.000000,bInvert=False))
 +AxisConfig=(AxisKeyName="Gamepad_Special_Left_X",AxisProperties=(DeadZone=0.000000,Sensitivity=1.000000,Exponent=1.000000,bInvert=False))
 +AxisConfig=(AxisKeyName="Gamepad_Special_Left_Y",AxisProperties=(DeadZone=0.000000,Sensitivity=1.000000,Exponent=1.000000,bInvert=False))
++AxisConfig=(AxisKeyName="Daydream_Left_Trackpad_X",AxisProperties=(DeadZone=0.000000,Sensitivity=1.000000,Exponent=1.000000,bInvert=False))
++AxisConfig=(AxisKeyName="Daydream_Left_Trackpad_Y",AxisProperties=(DeadZone=0.000000,Sensitivity=1.000000,Exponent=1.000000,bInvert=False))
++AxisConfig=(AxisKeyName="Daydream_Right_Trackpad_X",AxisProperties=(DeadZone=0.000000,Sensitivity=1.000000,Exponent=1.000000,bInvert=False))
++AxisConfig=(AxisKeyName="Daydream_Right_Trackpad_Y",AxisProperties=(DeadZone=0.000000,Sensitivity=1.000000,Exponent=1.000000,bInvert=False))
++AxisConfig=(AxisKeyName="Vive_Left_Trigger_Axis",AxisProperties=(DeadZone=0.000000,Sensitivity=1.000000,Exponent=1.000000,bInvert=False))
++AxisConfig=(AxisKeyName="Vive_Left_Trackpad_X",AxisProperties=(DeadZone=0.000000,Sensitivity=1.000000,Exponent=1.000000,bInvert=False))
++AxisConfig=(AxisKeyName="Vive_Left_Trackpad_Y",AxisProperties=(DeadZone=0.000000,Sensitivity=1.000000,Exponent=1.000000,bInvert=False))
++AxisConfig=(AxisKeyName="Vive_Right_Trigger_Axis",AxisProperties=(DeadZone=0.000000,Sensitivity=1.000000,Exponent=1.000000,bInvert=False))
++AxisConfig=(AxisKeyName="Vive_Right_Trackpad_X",AxisProperties=(DeadZone=0.000000,Sensitivity=1.000000,Exponent=1.000000,bInvert=False))
++AxisConfig=(AxisKeyName="Vive_Right_Trackpad_Y",AxisProperties=(DeadZone=0.000000,Sensitivity=1.000000,Exponent=1.000000,bInvert=False))
++AxisConfig=(AxisKeyName="MixedReality_Left_Trigger_Axis",AxisProperties=(DeadZone=0.000000,Sensitivity=1.000000,Exponent=1.000000,bInvert=False))
++AxisConfig=(AxisKeyName="MixedReality_Left_Thumbstick_X",AxisProperties=(DeadZone=0.000000,Sensitivity=1.000000,Exponent=1.000000,bInvert=False))
++AxisConfig=(AxisKeyName="MixedReality_Left_Thumbstick_Y",AxisProperties=(DeadZone=0.000000,Sensitivity=1.000000,Exponent=1.000000,bInvert=False))
++AxisConfig=(AxisKeyName="MixedReality_Left_Trackpad_X",AxisProperties=(DeadZone=0.000000,Sensitivity=1.000000,Exponent=1.000000,bInvert=False))
++AxisConfig=(AxisKeyName="MixedReality_Left_Trackpad_Y",AxisProperties=(DeadZone=0.000000,Sensitivity=1.000000,Exponent=1.000000,bInvert=False))
++AxisConfig=(AxisKeyName="MixedReality_Right_Trigger_Axis",AxisProperties=(DeadZone=0.000000,Sensitivity=1.000000,Exponent=1.000000,bInvert=False))
++AxisConfig=(AxisKeyName="MixedReality_Right_Thumbstick_X",AxisProperties=(DeadZone=0.000000,Sensitivity=1.000000,Exponent=1.000000,bInvert=False))
++AxisConfig=(AxisKeyName="MixedReality_Right_Thumbstick_Y",AxisProperties=(DeadZone=0.000000,Sensitivity=1.000000,Exponent=1.000000,bInvert=False))
++AxisConfig=(AxisKeyName="MixedReality_Right_Trackpad_X",AxisProperties=(DeadZone=0.000000,Sensitivity=1.000000,Exponent=1.000000,bInvert=False))
++AxisConfig=(AxisKeyName="MixedReality_Right_Trackpad_Y",AxisProperties=(DeadZone=0.000000,Sensitivity=1.000000,Exponent=1.000000,bInvert=False))
++AxisConfig=(AxisKeyName="OculusGo_Left_Trackpad_X",AxisProperties=(DeadZone=0.000000,Sensitivity=1.000000,Exponent=1.000000,bInvert=False))
++AxisConfig=(AxisKeyName="OculusGo_Left_Trackpad_Y",AxisProperties=(DeadZone=0.000000,Sensitivity=1.000000,Exponent=1.000000,bInvert=False))
++AxisConfig=(AxisKeyName="OculusGo_Right_Trackpad_X",AxisProperties=(DeadZone=0.000000,Sensitivity=1.000000,Exponent=1.000000,bInvert=False))
++AxisConfig=(AxisKeyName="OculusGo_Right_Trackpad_Y",AxisProperties=(DeadZone=0.000000,Sensitivity=1.000000,Exponent=1.000000,bInvert=False))
++AxisConfig=(AxisKeyName="OculusTouch_Left_Grip_Axis",AxisProperties=(DeadZone=0.000000,Sensitivity=1.000000,Exponent=1.000000,bInvert=False))
++AxisConfig=(AxisKeyName="OculusTouch_Left_Trigger_Axis",AxisProperties=(DeadZone=0.000000,Sensitivity=1.000000,Exponent=1.000000,bInvert=False))
++AxisConfig=(AxisKeyName="OculusTouch_Left_Thumbstick_X",AxisProperties=(DeadZone=0.000000,Sensitivity=1.000000,Exponent=1.000000,bInvert=False))
++AxisConfig=(AxisKeyName="OculusTouch_Left_Thumbstick_Y",AxisProperties=(DeadZone=0.000000,Sensitivity=1.000000,Exponent=1.000000,bInvert=False))
++AxisConfig=(AxisKeyName="OculusTouch_Right_Grip_Axis",AxisProperties=(DeadZone=0.000000,Sensitivity=1.000000,Exponent=1.000000,bInvert=False))
++AxisConfig=(AxisKeyName="OculusTouch_Right_Trigger_Axis",AxisProperties=(DeadZone=0.000000,Sensitivity=1.000000,Exponent=1.000000,bInvert=False))
++AxisConfig=(AxisKeyName="OculusTouch_Right_Thumbstick_X",AxisProperties=(DeadZone=0.000000,Sensitivity=1.000000,Exponent=1.000000,bInvert=False))
++AxisConfig=(AxisKeyName="OculusTouch_Right_Thumbstick_Y",AxisProperties=(DeadZone=0.000000,Sensitivity=1.000000,Exponent=1.000000,bInvert=False))
++AxisConfig=(AxisKeyName="ValveIndex_Left_Grip_Axis",AxisProperties=(DeadZone=0.000000,Sensitivity=1.000000,Exponent=1.000000,bInvert=False))
++AxisConfig=(AxisKeyName="ValveIndex_Left_Grip_Force",AxisProperties=(DeadZone=0.000000,Sensitivity=1.000000,Exponent=1.000000,bInvert=False))
++AxisConfig=(AxisKeyName="ValveIndex_Left_Trigger_Axis",AxisProperties=(DeadZone=0.000000,Sensitivity=1.000000,Exponent=1.000000,bInvert=False))
++AxisConfig=(AxisKeyName="ValveIndex_Left_Thumbstick_X",AxisProperties=(DeadZone=0.000000,Sensitivity=1.000000,Exponent=1.000000,bInvert=False))
++AxisConfig=(AxisKeyName="ValveIndex_Left_Thumbstick_Y",AxisProperties=(DeadZone=0.000000,Sensitivity=1.000000,Exponent=1.000000,bInvert=False))
++AxisConfig=(AxisKeyName="ValveIndex_Left_Trackpad_X",AxisProperties=(DeadZone=0.000000,Sensitivity=1.000000,Exponent=1.000000,bInvert=False))
++AxisConfig=(AxisKeyName="ValveIndex_Left_Trackpad_Y",AxisProperties=(DeadZone=0.000000,Sensitivity=1.000000,Exponent=1.000000,bInvert=False))
++AxisConfig=(AxisKeyName="ValveIndex_Left_Trackpad_Force",AxisProperties=(DeadZone=0.000000,Sensitivity=1.000000,Exponent=1.000000,bInvert=False))
++AxisConfig=(AxisKeyName="ValveIndex_Left_Trackpad_Touch",AxisProperties=(DeadZone=0.000000,Sensitivity=1.000000,Exponent=1.000000,bInvert=False))
++AxisConfig=(AxisKeyName="ValveIndex_Right_Grip_Axis",AxisProperties=(DeadZone=0.000000,Sensitivity=1.000000,Exponent=1.000000,bInvert=False))
++AxisConfig=(AxisKeyName="ValveIndex_Right_Grip_Force",AxisProperties=(DeadZone=0.000000,Sensitivity=1.000000,Exponent=1.000000,bInvert=False))
++AxisConfig=(AxisKeyName="ValveIndex_Right_Trigger_Axis",AxisProperties=(DeadZone=0.000000,Sensitivity=1.000000,Exponent=1.000000,bInvert=False))
++AxisConfig=(AxisKeyName="ValveIndex_Right_Thumbstick_X",AxisProperties=(DeadZone=0.000000,Sensitivity=1.000000,Exponent=1.000000,bInvert=False))
++AxisConfig=(AxisKeyName="ValveIndex_Right_Thumbstick_Y",AxisProperties=(DeadZone=0.000000,Sensitivity=1.000000,Exponent=1.000000,bInvert=False))
++AxisConfig=(AxisKeyName="ValveIndex_Right_Trackpad_X",AxisProperties=(DeadZone=0.000000,Sensitivity=1.000000,Exponent=1.000000,bInvert=False))
++AxisConfig=(AxisKeyName="ValveIndex_Right_Trackpad_Y",AxisProperties=(DeadZone=0.000000,Sensitivity=1.000000,Exponent=1.000000,bInvert=False))
++AxisConfig=(AxisKeyName="ValveIndex_Right_Trackpad_Force",AxisProperties=(DeadZone=0.000000,Sensitivity=1.000000,Exponent=1.000000,bInvert=False))
++AxisConfig=(AxisKeyName="MotionController_Left_Thumbstick_Z",AxisProperties=(DeadZone=0.000000,Sensitivity=1.000000,Exponent=1.000000,bInvert=False))
++AxisConfig=(AxisKeyName="MagicLeap_Left_Trigger_Axis",AxisProperties=(DeadZone=0.000000,Sensitivity=1.000000,Exponent=1.000000,bInvert=False))
++AxisConfig=(AxisKeyName="MagicLeap_Left_Trackpad_X",AxisProperties=(DeadZone=0.000000,Sensitivity=1.000000,Exponent=1.000000,bInvert=False))
++AxisConfig=(AxisKeyName="MagicLeap_Left_Trackpad_Y",AxisProperties=(DeadZone=0.000000,Sensitivity=1.000000,Exponent=1.000000,bInvert=False))
++AxisConfig=(AxisKeyName="MagicLeap_Left_Trackpad_Force",AxisProperties=(DeadZone=0.000000,Sensitivity=1.000000,Exponent=1.000000,bInvert=False))
++AxisConfig=(AxisKeyName="MagicLeap_Left_Touch1_X",AxisProperties=(DeadZone=0.000000,Sensitivity=1.000000,Exponent=1.000000,bInvert=False))
++AxisConfig=(AxisKeyName="MagicLeap_Left_Touch1_Y",AxisProperties=(DeadZone=0.000000,Sensitivity=1.000000,Exponent=1.000000,bInvert=False))
++AxisConfig=(AxisKeyName="MagicLeap_Left_Touch1_Force",AxisProperties=(DeadZone=0.000000,Sensitivity=1.000000,Exponent=1.000000,bInvert=False))
++AxisConfig=(AxisKeyName="MotionController_Right_Thumbstick_Z",AxisProperties=(DeadZone=0.000000,Sensitivity=1.000000,Exponent=1.000000,bInvert=False))
++AxisConfig=(AxisKeyName="MagicLeap_Right_Trigger_Axis",AxisProperties=(DeadZone=0.000000,Sensitivity=1.000000,Exponent=1.000000,bInvert=False))
++AxisConfig=(AxisKeyName="MagicLeap_Right_Trackpad_X",AxisProperties=(DeadZone=0.000000,Sensitivity=1.000000,Exponent=1.000000,bInvert=False))
++AxisConfig=(AxisKeyName="MagicLeap_Right_Trackpad_Y",AxisProperties=(DeadZone=0.000000,Sensitivity=1.000000,Exponent=1.000000,bInvert=False))
++AxisConfig=(AxisKeyName="MagicLeap_Right_Trackpad_Force",AxisProperties=(DeadZone=0.000000,Sensitivity=1.000000,Exponent=1.000000,bInvert=False))
++AxisConfig=(AxisKeyName="MagicLeap_Right_Touch1_X",AxisProperties=(DeadZone=0.000000,Sensitivity=1.000000,Exponent=1.000000,bInvert=False))
++AxisConfig=(AxisKeyName="MagicLeap_Right_Touch1_Y",AxisProperties=(DeadZone=0.000000,Sensitivity=1.000000,Exponent=1.000000,bInvert=False))
++AxisConfig=(AxisKeyName="MagicLeap_Right_Touch1_Force",AxisProperties=(DeadZone=0.000000,Sensitivity=1.000000,Exponent=1.000000,bInvert=False))
++AxisConfig=(AxisKeyName="OculusTouch_Left_Thumbstick",AxisProperties=(DeadZone=0.000000,Sensitivity=1.000000,Exponent=1.000000,bInvert=False))
++AxisConfig=(AxisKeyName="OculusTouch_Left_FaceButton1",AxisProperties=(DeadZone=0.000000,Sensitivity=1.000000,Exponent=1.000000,bInvert=False))
++AxisConfig=(AxisKeyName="OculusTouch_Left_Trigger",AxisProperties=(DeadZone=0.000000,Sensitivity=1.000000,Exponent=1.000000,bInvert=False))
++AxisConfig=(AxisKeyName="OculusTouch_Left_FaceButton2",AxisProperties=(DeadZone=0.000000,Sensitivity=1.000000,Exponent=1.000000,bInvert=False))
++AxisConfig=(AxisKeyName="OculusTouch_Left_IndexPointing",AxisProperties=(DeadZone=0.000000,Sensitivity=1.000000,Exponent=1.000000,bInvert=False))
++AxisConfig=(AxisKeyName="OculusTouch_Left_ThumbUp",AxisProperties=(DeadZone=0.000000,Sensitivity=1.000000,Exponent=1.000000,bInvert=False))
++AxisConfig=(AxisKeyName="OculusTouch_Right_Thumbstick",AxisProperties=(DeadZone=0.000000,Sensitivity=1.000000,Exponent=1.000000,bInvert=False))
++AxisConfig=(AxisKeyName="OculusTouch_Right_FaceButton1",AxisProperties=(DeadZone=0.000000,Sensitivity=1.000000,Exponent=1.000000,bInvert=False))
++AxisConfig=(AxisKeyName="OculusTouch_Right_Trigger",AxisProperties=(DeadZone=0.000000,Sensitivity=1.000000,Exponent=1.000000,bInvert=False))
++AxisConfig=(AxisKeyName="OculusTouch_Right_FaceButton2",AxisProperties=(DeadZone=0.000000,Sensitivity=1.000000,Exponent=1.000000,bInvert=False))
++AxisConfig=(AxisKeyName="OculusTouch_Right_IndexPointing",AxisProperties=(DeadZone=0.000000,Sensitivity=1.000000,Exponent=1.000000,bInvert=False))
++AxisConfig=(AxisKeyName="OculusTouch_Right_ThumbUp",AxisProperties=(DeadZone=0.000000,Sensitivity=1.000000,Exponent=1.000000,bInvert=False))
++AxisConfig=(AxisKeyName="OculusTouchpad_Touchpad_X",AxisProperties=(DeadZone=0.000000,Sensitivity=1.000000,Exponent=1.000000,bInvert=False))
++AxisConfig=(AxisKeyName="OculusTouchpad_Touchpad_Y",AxisProperties=(DeadZone=0.000000,Sensitivity=1.000000,Exponent=1.000000,bInvert=False))
 bAltEnterTogglesFullscreen=True
+bF11TogglesFullscreen=True
 bUseMouseForTouch=False
 bEnableMouseSmoothing=True
 bEnableFOVScaling=True
-FOVScale=0.011110
-DoubleClickTime=0.200000
 bCaptureMouseOnLaunch=True
-DefaultViewportMouseCaptureMode=CapturePermanently_IncludingInitialMouseDown
-bDefaultViewportMouseLock=True
-+ActionMappings=(ActionName="SetDestination",Key=LeftMouseButton,bShift=False,bCtrl=False,bAlt=False,bCmd=False)
-+ActionMappings=(ActionName="SetDestination",Key=MotionController_Right_Trigger,bShift=False,bCtrl=False,bAlt=False,bCmd=False)
-+ActionMappings=(ActionName="SetDestination",Key=Gamepad_RightTrigger,bShift=False,bCtrl=False,bAlt=False,bCmd=False)
-+ActionMappings=(ActionName="ResetVR",Key=R,bShift=False,bCtrl=False,bAlt=False,bCmd=False)
-+ActionMappings=(ActionName="ResetVR",Key=MotionController_Left_Grip1,bShift=False,bCtrl=False,bAlt=False,bCmd=False)
-+AxisMappings=(AxisName="MoveForward",Key=W,Scale=1.000000)
-+AxisMappings=(AxisName="MoveForward",Key=S,Scale=-1.000000)
-+AxisMappings=(AxisName="MoveRight",Key=D,Scale=1.000000)
-+AxisMappings=(AxisName="MoveRight",Key=A,Scale=-1.000000)
-+AxisMappings=(AxisName="TurnRate",Key=MotionController_Right_Thumbstick_X,Scale=1.000000)
-+AxisMappings=(AxisName="TurnRate",Key=Gamepad_RightX,Scale=1.000000)
 bAlwaysShowTouchInterface=False
 bShowConsoleOnFourFingerTap=True
+bEnableGestureRecognizer=False
+bUseAutocorrect=False
+DefaultViewportMouseCaptureMode=CapturePermanently_IncludingInitialMouseDown
+DefaultViewportMouseLockMode=LockOnCapture
+FOVScale=0.011110
+DoubleClickTime=0.200000
++ActionMappings=(ActionName="SetDestination",bShift=False,bCtrl=False,bAlt=False,bCmd=False,Key=LeftMouseButton)
++ActionMappings=(ActionName="SetDestination",bShift=False,bCtrl=False,bAlt=False,bCmd=False,Key=Gamepad_RightTrigger)
++ActionMappings=(ActionName="ResetVR",bShift=False,bCtrl=False,bAlt=False,bCmd=False,Key=R)
++ActionMappings=(ActionName="ResetVR",bShift=False,bCtrl=False,bAlt=False,bCmd=False,Key=Daydream_Left_Trackpad_Click)
++ActionMappings=(ActionName="SetDestination",bShift=False,bCtrl=False,bAlt=False,bCmd=False,Key=Daydream_Left_Select_Click)
++ActionMappings=(ActionName="SetDestination",bShift=False,bCtrl=False,bAlt=False,bCmd=False,Key=Vive_Left_Trigger_Click)
++ActionMappings=(ActionName="SetDestination",bShift=False,bCtrl=False,bAlt=False,bCmd=False,Key=MixedReality_Left_Trigger_Click)
++ActionMappings=(ActionName="SetDestination",bShift=False,bCtrl=False,bAlt=False,bCmd=False,Key=OculusGo_Left_Trigger_Click)
++ActionMappings=(ActionName="SetDestination",bShift=False,bCtrl=False,bAlt=False,bCmd=False,Key=OculusTouch_Left_Trigger_Click)
++ActionMappings=(ActionName="SetDestination",bShift=False,bCtrl=False,bAlt=False,bCmd=False,Key=ValveIndex_Left_Trigger_Click)
++ActionMappings=(ActionName="ResetVR",bShift=False,bCtrl=False,bAlt=False,bCmd=False,Key=Vive_Left_Grip_Click)
++ActionMappings=(ActionName="ResetVR",bShift=False,bCtrl=False,bAlt=False,bCmd=False,Key=MixedReality_Left_Thumbstick_Click)
++ActionMappings=(ActionName="ResetVR",bShift=False,bCtrl=False,bAlt=False,bCmd=False,Key=OculusGo_Left_Back_Click)
++ActionMappings=(ActionName="ResetVR",bShift=False,bCtrl=False,bAlt=False,bCmd=False,Key=OculusTouch_Left_Thumbstick_Click)
++ActionMappings=(ActionName="ResetVR",bShift=False,bCtrl=False,bAlt=False,bCmd=False,Key=ValveIndex_Left_Thumbstick_Click)
++ActionMappings=(ActionName="ResetVR",bShift=False,bCtrl=False,bAlt=False,bCmd=False,Key=MagicLeap_Left_Bumper)
++ActionMappings=(ActionName="SetDestination",bShift=False,bCtrl=False,bAlt=False,bCmd=False,Key=MagicLeap_Left_Trigger)
++AxisMappings=(AxisName="MoveForward",Scale=1.000000,Key=W)
++AxisMappings=(AxisName="MoveForward",Scale=-1.000000,Key=S)
++AxisMappings=(AxisName="MoveRight",Scale=1.000000,Key=D)
++AxisMappings=(AxisName="MoveRight",Scale=-1.000000,Key=A)
++AxisMappings=(AxisName="TurnRate",Scale=1.000000,Key=Gamepad_RightX)
++AxisMappings=(AxisName="TurnRate",Scale=1.000000,Key=Daydream_Left_Trackpad_X)
++AxisMappings=(AxisName="TurnRate",Scale=1.000000,Key=Vive_Left_Trackpad_X)
++AxisMappings=(AxisName="TurnRate",Scale=1.000000,Key=MixedReality_Left_Thumbstick_X)
++AxisMappings=(AxisName="TurnRate",Scale=1.000000,Key=OculusGo_Left_Trackpad_X)
++AxisMappings=(AxisName="TurnRate",Scale=1.000000,Key=OculusTouch_Left_Thumbstick_X)
++AxisMappings=(AxisName="TurnRate",Scale=1.000000,Key=ValveIndex_Left_Thumbstick_X)
++AxisMappings=(AxisName="TurnRate",Scale=1.000000,Key=MagicLeap_Left_Trackpad_X)
 DefaultTouchInterface=None
-ConsoleKey=None
--ConsoleKeys=Tilde
 +ConsoleKeys=Tilde
 
 
diff --git a/Templates/TP_TopDownBP/Config/DefaultInput.ini b/Templates/TP_TopDownBP/Config/DefaultInput.ini
--- a/Templates/TP_TopDownBP/Config/DefaultInput.ini
+++ b/Templates/TP_TopDownBP/Config/DefaultInput.ini
@@ -1,11 +1,5 @@
 
 [/Script/Engine.InputSettings]
--AxisConfig=(AxisKeyName="Gamepad_LeftX",AxisProperties=(DeadZone=0.25,Exponent=1.f,Sensitivity=1.f))
--AxisConfig=(AxisKeyName="Gamepad_LeftY",AxisProperties=(DeadZone=0.25,Exponent=1.f,Sensitivity=1.f))
--AxisConfig=(AxisKeyName="Gamepad_RightX",AxisProperties=(DeadZone=0.25,Exponent=1.f,Sensitivity=1.f))
--AxisConfig=(AxisKeyName="Gamepad_RightY",AxisProperties=(DeadZone=0.25,Exponent=1.f,Sensitivity=1.f))
--AxisConfig=(AxisKeyName="MouseX",AxisProperties=(DeadZone=0.f,Exponent=1.f,Sensitivity=0.07f))
--AxisConfig=(AxisKeyName="MouseY",AxisProperties=(DeadZone=0.f,Exponent=1.f,Sensitivity=0.07f))
 +AxisConfig=(AxisKeyName="Gamepad_LeftX",AxisProperties=(DeadZone=0.250000,Sensitivity=1.000000,Exponent=1.000000,bInvert=False))
 +AxisConfig=(AxisKeyName="Gamepad_LeftY",AxisProperties=(DeadZone=0.250000,Sensitivity=1.000000,Exponent=1.000000,bInvert=False))
 +AxisConfig=(AxisKeyName="Gamepad_RightX",AxisProperties=(DeadZone=0.250000,Sensitivity=1.000000,Exponent=1.000000,bInvert=False))
@@ -27,31 +21,129 @@
 +AxisConfig=(AxisKeyName="MotionController_Right_Grip2Axis",AxisProperties=(DeadZone=0.000000,Sensitivity=1.000000,Exponent=1.000000,bInvert=False))
 +AxisConfig=(AxisKeyName="Gamepad_Special_Left_X",AxisProperties=(DeadZone=0.000000,Sensitivity=1.000000,Exponent=1.000000,bInvert=False))
 +AxisConfig=(AxisKeyName="Gamepad_Special_Left_Y",AxisProperties=(DeadZone=0.000000,Sensitivity=1.000000,Exponent=1.000000,bInvert=False))
++AxisConfig=(AxisKeyName="Daydream_Left_Trackpad_X",AxisProperties=(DeadZone=0.000000,Sensitivity=1.000000,Exponent=1.000000,bInvert=False))
++AxisConfig=(AxisKeyName="Daydream_Left_Trackpad_Y",AxisProperties=(DeadZone=0.000000,Sensitivity=1.000000,Exponent=1.000000,bInvert=False))
++AxisConfig=(AxisKeyName="Daydream_Right_Trackpad_X",AxisProperties=(DeadZone=0.000000,Sensitivity=1.000000,Exponent=1.000000,bInvert=False))
++AxisConfig=(AxisKeyName="Daydream_Right_Trackpad_Y",AxisProperties=(DeadZone=0.000000,Sensitivity=1.000000,Exponent=1.000000,bInvert=False))
++AxisConfig=(AxisKeyName="Vive_Left_Trigger_Axis",AxisProperties=(DeadZone=0.000000,Sensitivity=1.000000,Exponent=1.000000,bInvert=False))
++AxisConfig=(AxisKeyName="Vive_Left_Trackpad_X",AxisProperties=(DeadZone=0.000000,Sensitivity=1.000000,Exponent=1.000000,bInvert=False))
++AxisConfig=(AxisKeyName="Vive_Left_Trackpad_Y",AxisProperties=(DeadZone=0.000000,Sensitivity=1.000000,Exponent=1.000000,bInvert=False))
++AxisConfig=(AxisKeyName="Vive_Right_Trigger_Axis",AxisProperties=(DeadZone=0.000000,Sensitivity=1.000000,Exponent=1.000000,bInvert=False))
++AxisConfig=(AxisKeyName="Vive_Right_Trackpad_X",AxisProperties=(DeadZone=0.000000,Sensitivity=1.000000,Exponent=1.000000,bInvert=False))
++AxisConfig=(AxisKeyName="Vive_Right_Trackpad_Y",AxisProperties=(DeadZone=0.000000,Sensitivity=1.000000,Exponent=1.000000,bInvert=False))
++AxisConfig=(AxisKeyName="MixedReality_Left_Trigger_Axis",AxisProperties=(DeadZone=0.000000,Sensitivity=1.000000,Exponent=1.000000,bInvert=False))
++AxisConfig=(AxisKeyName="MixedReality_Left_Thumbstick_X",AxisProperties=(DeadZone=0.000000,Sensitivity=1.000000,Exponent=1.000000,bInvert=False))
++AxisConfig=(AxisKeyName="MixedReality_Left_Thumbstick_Y",AxisProperties=(DeadZone=0.000000,Sensitivity=1.000000,Exponent=1.000000,bInvert=False))
++AxisConfig=(AxisKeyName="MixedReality_Left_Trackpad_X",AxisProperties=(DeadZone=0.000000,Sensitivity=1.000000,Exponent=1.000000,bInvert=False))
++AxisConfig=(AxisKeyName="MixedReality_Left_Trackpad_Y",AxisProperties=(DeadZone=0.000000,Sensitivity=1.000000,Exponent=1.000000,bInvert=False))
++AxisConfig=(AxisKeyName="MixedReality_Right_Trigger_Axis",AxisProperties=(DeadZone=0.000000,Sensitivity=1.000000,Exponent=1.000000,bInvert=False))
++AxisConfig=(AxisKeyName="MixedReality_Right_Thumbstick_X",AxisProperties=(DeadZone=0.000000,Sensitivity=1.000000,Exponent=1.000000,bInvert=False))
++AxisConfig=(AxisKeyName="MixedReality_Right_Thumbstick_Y",AxisProperties=(DeadZone=0.000000,Sensitivity=1.000000,Exponent=1.000000,bInvert=False))
++AxisConfig=(AxisKeyName="MixedReality_Right_Trackpad_X",AxisProperties=(DeadZone=0.000000,Sensitivity=1.000000,Exponent=1.000000,bInvert=False))
++AxisConfig=(AxisKeyName="MixedReality_Right_Trackpad_Y",AxisProperties=(DeadZone=0.000000,Sensitivity=1.000000,Exponent=1.000000,bInvert=False))
++AxisConfig=(AxisKeyName="OculusGo_Left_Trackpad_X",AxisProperties=(DeadZone=0.000000,Sensitivity=1.000000,Exponent=1.000000,bInvert=False))
++AxisConfig=(AxisKeyName="OculusGo_Left_Trackpad_Y",AxisProperties=(DeadZone=0.000000,Sensitivity=1.000000,Exponent=1.000000,bInvert=False))
++AxisConfig=(AxisKeyName="OculusGo_Right_Trackpad_X",AxisProperties=(DeadZone=0.000000,Sensitivity=1.000000,Exponent=1.000000,bInvert=False))
++AxisConfig=(AxisKeyName="OculusGo_Right_Trackpad_Y",AxisProperties=(DeadZone=0.000000,Sensitivity=1.000000,Exponent=1.000000,bInvert=False))
++AxisConfig=(AxisKeyName="OculusTouch_Left_Grip_Axis",AxisProperties=(DeadZone=0.000000,Sensitivity=1.000000,Exponent=1.000000,bInvert=False))
++AxisConfig=(AxisKeyName="OculusTouch_Left_Trigger_Axis",AxisProperties=(DeadZone=0.000000,Sensitivity=1.000000,Exponent=1.000000,bInvert=False))
++AxisConfig=(AxisKeyName="OculusTouch_Left_Thumbstick_X",AxisProperties=(DeadZone=0.000000,Sensitivity=1.000000,Exponent=1.000000,bInvert=False))
++AxisConfig=(AxisKeyName="OculusTouch_Left_Thumbstick_Y",AxisProperties=(DeadZone=0.000000,Sensitivity=1.000000,Exponent=1.000000,bInvert=False))
++AxisConfig=(AxisKeyName="OculusTouch_Right_Grip_Axis",AxisProperties=(DeadZone=0.000000,Sensitivity=1.000000,Exponent=1.000000,bInvert=False))
++AxisConfig=(AxisKeyName="OculusTouch_Right_Trigger_Axis",AxisProperties=(DeadZone=0.000000,Sensitivity=1.000000,Exponent=1.000000,bInvert=False))
++AxisConfig=(AxisKeyName="OculusTouch_Right_Thumbstick_X",AxisProperties=(DeadZone=0.000000,Sensitivity=1.000000,Exponent=1.000000,bInvert=False))
++AxisConfig=(AxisKeyName="OculusTouch_Right_Thumbstick_Y",AxisProperties=(DeadZone=0.000000,Sensitivity=1.000000,Exponent=1.000000,bInvert=False))
++AxisConfig=(AxisKeyName="ValveIndex_Left_Grip_Axis",AxisProperties=(DeadZone=0.000000,Sensitivity=1.000000,Exponent=1.000000,bInvert=False))
++AxisConfig=(AxisKeyName="ValveIndex_Left_Grip_Force",AxisProperties=(DeadZone=0.000000,Sensitivity=1.000000,Exponent=1.000000,bInvert=False))
++AxisConfig=(AxisKeyName="ValveIndex_Left_Trigger_Axis",AxisProperties=(DeadZone=0.000000,Sensitivity=1.000000,Exponent=1.000000,bInvert=False))
++AxisConfig=(AxisKeyName="ValveIndex_Left_Thumbstick_X",AxisProperties=(DeadZone=0.000000,Sensitivity=1.000000,Exponent=1.000000,bInvert=False))
++AxisConfig=(AxisKeyName="ValveIndex_Left_Thumbstick_Y",AxisProperties=(DeadZone=0.000000,Sensitivity=1.000000,Exponent=1.000000,bInvert=False))
++AxisConfig=(AxisKeyName="ValveIndex_Left_Trackpad_X",AxisProperties=(DeadZone=0.000000,Sensitivity=1.000000,Exponent=1.000000,bInvert=False))
++AxisConfig=(AxisKeyName="ValveIndex_Left_Trackpad_Y",AxisProperties=(DeadZone=0.000000,Sensitivity=1.000000,Exponent=1.000000,bInvert=False))
++AxisConfig=(AxisKeyName="ValveIndex_Left_Trackpad_Force",AxisProperties=(DeadZone=0.000000,Sensitivity=1.000000,Exponent=1.000000,bInvert=False))
++AxisConfig=(AxisKeyName="ValveIndex_Left_Trackpad_Touch",AxisProperties=(DeadZone=0.000000,Sensitivity=1.000000,Exponent=1.000000,bInvert=False))
++AxisConfig=(AxisKeyName="ValveIndex_Right_Grip_Axis",AxisProperties=(DeadZone=0.000000,Sensitivity=1.000000,Exponent=1.000000,bInvert=False))
++AxisConfig=(AxisKeyName="ValveIndex_Right_Grip_Force",AxisProperties=(DeadZone=0.000000,Sensitivity=1.000000,Exponent=1.000000,bInvert=False))
++AxisConfig=(AxisKeyName="ValveIndex_Right_Trigger_Axis",AxisProperties=(DeadZone=0.000000,Sensitivity=1.000000,Exponent=1.000000,bInvert=False))
++AxisConfig=(AxisKeyName="ValveIndex_Right_Thumbstick_X",AxisProperties=(DeadZone=0.000000,Sensitivity=1.000000,Exponent=1.000000,bInvert=False))
++AxisConfig=(AxisKeyName="ValveIndex_Right_Thumbstick_Y",AxisProperties=(DeadZone=0.000000,Sensitivity=1.000000,Exponent=1.000000,bInvert=False))
++AxisConfig=(AxisKeyName="ValveIndex_Right_Trackpad_X",AxisProperties=(DeadZone=0.000000,Sensitivity=1.000000,Exponent=1.000000,bInvert=False))
++AxisConfig=(AxisKeyName="ValveIndex_Right_Trackpad_Y",AxisProperties=(DeadZone=0.000000,Sensitivity=1.000000,Exponent=1.000000,bInvert=False))
++AxisConfig=(AxisKeyName="ValveIndex_Right_Trackpad_Force",AxisProperties=(DeadZone=0.000000,Sensitivity=1.000000,Exponent=1.000000,bInvert=False))
++AxisConfig=(AxisKeyName="MotionController_Left_Thumbstick_Z",AxisProperties=(DeadZone=0.000000,Sensitivity=1.000000,Exponent=1.000000,bInvert=False))
++AxisConfig=(AxisKeyName="MagicLeap_Left_Trigger_Axis",AxisProperties=(DeadZone=0.000000,Sensitivity=1.000000,Exponent=1.000000,bInvert=False))
++AxisConfig=(AxisKeyName="MagicLeap_Left_Trackpad_X",AxisProperties=(DeadZone=0.000000,Sensitivity=1.000000,Exponent=1.000000,bInvert=False))
++AxisConfig=(AxisKeyName="MagicLeap_Left_Trackpad_Y",AxisProperties=(DeadZone=0.000000,Sensitivity=1.000000,Exponent=1.000000,bInvert=False))
++AxisConfig=(AxisKeyName="MagicLeap_Left_Trackpad_Force",AxisProperties=(DeadZone=0.000000,Sensitivity=1.000000,Exponent=1.000000,bInvert=False))
++AxisConfig=(AxisKeyName="MagicLeap_Left_Touch1_X",AxisProperties=(DeadZone=0.000000,Sensitivity=1.000000,Exponent=1.000000,bInvert=False))
++AxisConfig=(AxisKeyName="MagicLeap_Left_Touch1_Y",AxisProperties=(DeadZone=0.000000,Sensitivity=1.000000,Exponent=1.000000,bInvert=False))
++AxisConfig=(AxisKeyName="MagicLeap_Left_Touch1_Force",AxisProperties=(DeadZone=0.000000,Sensitivity=1.000000,Exponent=1.000000,bInvert=False))
++AxisConfig=(AxisKeyName="MotionController_Right_Thumbstick_Z",AxisProperties=(DeadZone=0.000000,Sensitivity=1.000000,Exponent=1.000000,bInvert=False))
++AxisConfig=(AxisKeyName="MagicLeap_Right_Trigger_Axis",AxisProperties=(DeadZone=0.000000,Sensitivity=1.000000,Exponent=1.000000,bInvert=False))
++AxisConfig=(AxisKeyName="MagicLeap_Right_Trackpad_X",AxisProperties=(DeadZone=0.000000,Sensitivity=1.000000,Exponent=1.000000,bInvert=False))
++AxisConfig=(AxisKeyName="MagicLeap_Right_Trackpad_Y",AxisProperties=(DeadZone=0.000000,Sensitivity=1.000000,Exponent=1.000000,bInvert=False))
++AxisConfig=(AxisKeyName="MagicLeap_Right_Trackpad_Force",AxisProperties=(DeadZone=0.000000,Sensitivity=1.000000,Exponent=1.000000,bInvert=False))
++AxisConfig=(AxisKeyName="MagicLeap_Right_Touch1_X",AxisProperties=(DeadZone=0.000000,Sensitivity=1.000000,Exponent=1.000000,bInvert=False))
++AxisConfig=(AxisKeyName="MagicLeap_Right_Touch1_Y",AxisProperties=(DeadZone=0.000000,Sensitivity=1.000000,Exponent=1.000000,bInvert=False))
++AxisConfig=(AxisKeyName="MagicLeap_Right_Touch1_Force",AxisProperties=(DeadZone=0.000000,Sensitivity=1.000000,Exponent=1.000000,bInvert=False))
++AxisConfig=(AxisKeyName="OculusTouch_Left_Thumbstick",AxisProperties=(DeadZone=0.000000,Sensitivity=1.000000,Exponent=1.000000,bInvert=False))
++AxisConfig=(AxisKeyName="OculusTouch_Left_FaceButton1",AxisProperties=(DeadZone=0.000000,Sensitivity=1.000000,Exponent=1.000000,bInvert=False))
++AxisConfig=(AxisKeyName="OculusTouch_Left_Trigger",AxisProperties=(DeadZone=0.000000,Sensitivity=1.000000,Exponent=1.000000,bInvert=False))
++AxisConfig=(AxisKeyName="OculusTouch_Left_FaceButton2",AxisProperties=(DeadZone=0.000000,Sensitivity=1.000000,Exponent=1.000000,bInvert=False))
++AxisConfig=(AxisKeyName="OculusTouch_Left_IndexPointing",AxisProperties=(DeadZone=0.000000,Sensitivity=1.000000,Exponent=1.000000,bInvert=False))
++AxisConfig=(AxisKeyName="OculusTouch_Left_ThumbUp",AxisProperties=(DeadZone=0.000000,Sensitivity=1.000000,Exponent=1.000000,bInvert=False))
++AxisConfig=(AxisKeyName="OculusTouch_Right_Thumbstick",AxisProperties=(DeadZone=0.000000,Sensitivity=1.000000,Exponent=1.000000,bInvert=False))
++AxisConfig=(AxisKeyName="OculusTouch_Right_FaceButton1",AxisProperties=(DeadZone=0.000000,Sensitivity=1.000000,Exponent=1.000000,bInvert=False))
++AxisConfig=(AxisKeyName="OculusTouch_Right_Trigger",AxisProperties=(DeadZone=0.000000,Sensitivity=1.000000,Exponent=1.000000,bInvert=False))
++AxisConfig=(AxisKeyName="OculusTouch_Right_FaceButton2",AxisProperties=(DeadZone=0.000000,Sensitivity=1.000000,Exponent=1.000000,bInvert=False))
++AxisConfig=(AxisKeyName="OculusTouch_Right_IndexPointing",AxisProperties=(DeadZone=0.000000,Sensitivity=1.000000,Exponent=1.000000,bInvert=False))
++AxisConfig=(AxisKeyName="OculusTouch_Right_ThumbUp",AxisProperties=(DeadZone=0.000000,Sensitivity=1.000000,Exponent=1.000000,bInvert=False))
++AxisConfig=(AxisKeyName="OculusTouchpad_Touchpad_X",AxisProperties=(DeadZone=0.000000,Sensitivity=1.000000,Exponent=1.000000,bInvert=False))
++AxisConfig=(AxisKeyName="OculusTouchpad_Touchpad_Y",AxisProperties=(DeadZone=0.000000,Sensitivity=1.000000,Exponent=1.000000,bInvert=False))
 bAltEnterTogglesFullscreen=True
+bF11TogglesFullscreen=True
 bUseMouseForTouch=False
 bEnableMouseSmoothing=True
 bEnableFOVScaling=True
-FOVScale=0.011110
-DoubleClickTime=0.200000
 bCaptureMouseOnLaunch=True
-DefaultViewportMouseCaptureMode=CapturePermanently_IncludingInitialMouseDown
-bDefaultViewportMouseLock=True
-+ActionMappings=(ActionName="SetDestination",Key=LeftMouseButton,bShift=False,bCtrl=False,bAlt=False,bCmd=False)
-+ActionMappings=(ActionName="SetDestination",Key=MotionController_Right_Trigger,bShift=False,bCtrl=False,bAlt=False,bCmd=False)
-+ActionMappings=(ActionName="SetDestination",Key=Gamepad_RightTrigger,bShift=False,bCtrl=False,bAlt=False,bCmd=False)
-+ActionMappings=(ActionName="ResetVR",Key=R,bShift=False,bCtrl=False,bAlt=False,bCmd=False)
-+ActionMappings=(ActionName="ResetVR",Key=MotionController_Left_Grip1,bShift=False,bCtrl=False,bAlt=False,bCmd=False)
-+AxisMappings=(AxisName="MoveForward",Key=W,Scale=1.000000)
-+AxisMappings=(AxisName="MoveForward",Key=S,Scale=-1.000000)
-+AxisMappings=(AxisName="MoveRight",Key=D,Scale=1.000000)
-+AxisMappings=(AxisName="MoveRight",Key=A,Scale=-1.000000)
-+AxisMappings=(AxisName="TurnRate",Key=MotionController_Right_Thumbstick_X,Scale=1.000000)
-+AxisMappings=(AxisName="TurnRate",Key=Gamepad_RightX,Scale=1.000000)
 bAlwaysShowTouchInterface=False
 bShowConsoleOnFourFingerTap=True
+bEnableGestureRecognizer=False
+bUseAutocorrect=False
+DefaultViewportMouseCaptureMode=CapturePermanently_IncludingInitialMouseDown
+DefaultViewportMouseLockMode=LockOnCapture
+FOVScale=0.011110
+DoubleClickTime=0.200000
++ActionMappings=(ActionName="SetDestination",bShift=False,bCtrl=False,bAlt=False,bCmd=False,Key=LeftMouseButton)
++ActionMappings=(ActionName="SetDestination",bShift=False,bCtrl=False,bAlt=False,bCmd=False,Key=Gamepad_RightTrigger)
++ActionMappings=(ActionName="ResetVR",bShift=False,bCtrl=False,bAlt=False,bCmd=False,Key=R)
++ActionMappings=(ActionName="ResetVR",bShift=False,bCtrl=False,bAlt=False,bCmd=False,Key=Daydream_Left_Trackpad_Click)
++ActionMappings=(ActionName="SetDestination",bShift=False,bCtrl=False,bAlt=False,bCmd=False,Key=Daydream_Left_Select_Click)
++ActionMappings=(ActionName="SetDestination",bShift=False,bCtrl=False,bAlt=False,bCmd=False,Key=Vive_Left_Trigger_Click)
++ActionMappings=(ActionName="SetDestination",bShift=False,bCtrl=False,bAlt=False,bCmd=False,Key=MixedReality_Left_Trigger_Click)
++ActionMappings=(ActionName="SetDestination",bShift=False,bCtrl=False,bAlt=False,bCmd=False,Key=OculusGo_Left_Trigger_Click)
++ActionMappings=(ActionName="SetDestination",bShift=False,bCtrl=False,bAlt=False,bCmd=False,Key=OculusTouch_Left_Trigger_Click)
++ActionMappings=(ActionName="SetDestination",bShift=False,bCtrl=False,bAlt=False,bCmd=False,Key=ValveIndex_Left_Trigger_Click)
++ActionMappings=(ActionName="ResetVR",bShift=False,bCtrl=False,bAlt=False,bCmd=False,Key=Vive_Left_Grip_Click)
++ActionMappings=(ActionName="ResetVR",bShift=False,bCtrl=False,bAlt=False,bCmd=False,Key=MixedReality_Left_Thumbstick_Click)
++ActionMappings=(ActionName="ResetVR",bShift=False,bCtrl=False,bAlt=False,bCmd=False,Key=OculusGo_Left_Back_Click)
++ActionMappings=(ActionName="ResetVR",bShift=False,bCtrl=False,bAlt=False,bCmd=False,Key=OculusTouch_Left_Thumbstick_Click)
++ActionMappings=(ActionName="ResetVR",bShift=False,bCtrl=False,bAlt=False,bCmd=False,Key=ValveIndex_Left_Thumbstick_Click)
++ActionMappings=(ActionName="ResetVR",bShift=False,bCtrl=False,bAlt=False,bCmd=False,Key=MagicLeap_Left_Bumper)
++ActionMappings=(ActionName="SetDestination",bShift=False,bCtrl=False,bAlt=False,bCmd=False,Key=MagicLeap_Left_Trigger)
++AxisMappings=(AxisName="MoveForward",Scale=1.000000,Key=W)
++AxisMappings=(AxisName="MoveForward",Scale=-1.000000,Key=S)
++AxisMappings=(AxisName="MoveRight",Scale=1.000000,Key=D)
++AxisMappings=(AxisName="MoveRight",Scale=-1.000000,Key=A)
++AxisMappings=(AxisName="TurnRate",Scale=1.000000,Key=Gamepad_RightX)
++AxisMappings=(AxisName="TurnRate",Scale=1.000000,Key=Daydream_Left_Trackpad_X)
++AxisMappings=(AxisName="TurnRate",Scale=1.000000,Key=Vive_Left_Trackpad_X)
++AxisMappings=(AxisName="TurnRate",Scale=1.000000,Key=MixedReality_Left_Thumbstick_X)
++AxisMappings=(AxisName="TurnRate",Scale=1.000000,Key=OculusGo_Left_Trackpad_X)
++AxisMappings=(AxisName="TurnRate",Scale=1.000000,Key=OculusTouch_Left_Thumbstick_X)
++AxisMappings=(AxisName="TurnRate",Scale=1.000000,Key=ValveIndex_Left_Thumbstick_X)
++AxisMappings=(AxisName="TurnRate",Scale=1.000000,Key=MagicLeap_Left_Trackpad_X)
 DefaultTouchInterface=None
-ConsoleKey=None
--ConsoleKeys=Tilde
 +ConsoleKeys=Tilde
 
 
diff --git a/Templates/TP_Vehicle/Config/DefaultInput.ini b/Templates/TP_Vehicle/Config/DefaultInput.ini
--- a/Templates/TP_Vehicle/Config/DefaultInput.ini
+++ b/Templates/TP_Vehicle/Config/DefaultInput.ini
@@ -1,61 +1,175 @@
+[/Script/RawInput.RawInputSettings]
++DeviceConfigurations=(VendorID="0x046D",ProductID="0xC262",AxisProperties=((bEnabled=False),(Key=GenericUSBController_Axis1,Offset=-0.500000),(bEnabled=False),(Key=GenericUSBController_Axis2,bInverted=True,Offset=1.000000),(Key=GenericUSBController_Axis3,bInverted=True,Offset=1.000000)),ButtonProperties=((Key=GenericUSBController_Button1),(Key=GenericUSBController_Button2),(Key=GenericUSBController_Button3),(Key=GenericUSBController_Button4),(Key=GenericUSBController_Button5),(Key=GenericUSBController_Button6),(Key=GenericUSBController_Button7),(Key=GenericUSBController_Button8),(Key=GenericUSBController_Button9),(Key=GenericUSBController_Button10),(Key=GenericUSBController_Button11),(Key=GenericUSBController_Button12)))
 
 [/Script/Engine.InputSettings]
--AxisConfig=(AxisKeyName="Gamepad_LeftX",AxisProperties=(DeadZone=0.25,Exponent=1.f,Sensitivity=1.f))
--AxisConfig=(AxisKeyName="Gamepad_LeftY",AxisProperties=(DeadZone=0.25,Exponent=1.f,Sensitivity=1.f))
--AxisConfig=(AxisKeyName="Gamepad_RightX",AxisProperties=(DeadZone=0.25,Exponent=1.f,Sensitivity=1.f))
--AxisConfig=(AxisKeyName="Gamepad_RightY",AxisProperties=(DeadZone=0.25,Exponent=1.f,Sensitivity=1.f))
--AxisConfig=(AxisKeyName="MouseX",AxisProperties=(DeadZone=0.f,Exponent=1.f,Sensitivity=0.07f))
--AxisConfig=(AxisKeyName="MouseY",AxisProperties=(DeadZone=0.f,Exponent=1.f,Sensitivity=0.07f))
-+AxisConfig=(AxisKeyName="Gamepad_LeftX",AxisProperties=(DeadZone=0.250000,Sensitivity=1.000000,Exponent=1.000000,bInvert=False))
-+AxisConfig=(AxisKeyName="Gamepad_LeftY",AxisProperties=(DeadZone=0.250000,Sensitivity=1.000000,Exponent=1.000000,bInvert=False))
-+AxisConfig=(AxisKeyName="Gamepad_RightX",AxisProperties=(DeadZone=0.250000,Sensitivity=1.000000,Exponent=1.000000,bInvert=False))
-+AxisConfig=(AxisKeyName="Gamepad_RightY",AxisProperties=(DeadZone=0.250000,Sensitivity=1.000000,Exponent=1.000000,bInvert=False))
-+AxisConfig=(AxisKeyName="MouseX",AxisProperties=(DeadZone=0.000000,Sensitivity=0.070000,Exponent=1.000000,bInvert=False))
-+AxisConfig=(AxisKeyName="MouseY",AxisProperties=(DeadZone=0.000000,Sensitivity=0.070000,Exponent=1.000000,bInvert=False))
 +AxisConfig=(AxisKeyName="Gamepad_LeftX",AxisProperties=(DeadZone=0.250000,Sensitivity=1.000000,Exponent=1.000000,bInvert=False))
 +AxisConfig=(AxisKeyName="Gamepad_LeftY",AxisProperties=(DeadZone=0.250000,Sensitivity=1.000000,Exponent=1.000000,bInvert=False))
 +AxisConfig=(AxisKeyName="Gamepad_RightX",AxisProperties=(DeadZone=0.250000,Sensitivity=1.000000,Exponent=1.000000,bInvert=False))
 +AxisConfig=(AxisKeyName="Gamepad_RightY",AxisProperties=(DeadZone=0.250000,Sensitivity=1.000000,Exponent=1.000000,bInvert=False))
 +AxisConfig=(AxisKeyName="MouseX",AxisProperties=(DeadZone=0.000000,Sensitivity=0.070000,Exponent=1.000000,bInvert=False))
 +AxisConfig=(AxisKeyName="MouseY",AxisProperties=(DeadZone=0.000000,Sensitivity=0.070000,Exponent=1.000000,bInvert=False))
++AxisConfig=(AxisKeyName="MouseWheelAxis",AxisProperties=(DeadZone=0.000000,Sensitivity=1.000000,Exponent=1.000000,bInvert=False))
++AxisConfig=(AxisKeyName="Gamepad_LeftTriggerAxis",AxisProperties=(DeadZone=0.000000,Sensitivity=1.000000,Exponent=1.000000,bInvert=False))
++AxisConfig=(AxisKeyName="Gamepad_RightTriggerAxis",AxisProperties=(DeadZone=0.000000,Sensitivity=1.000000,Exponent=1.000000,bInvert=False))
++AxisConfig=(AxisKeyName="MotionController_Left_Thumbstick_X",AxisProperties=(DeadZone=0.000000,Sensitivity=1.000000,Exponent=1.000000,bInvert=False))
++AxisConfig=(AxisKeyName="MotionController_Left_Thumbstick_Y",AxisProperties=(DeadZone=0.000000,Sensitivity=1.000000,Exponent=1.000000,bInvert=False))
++AxisConfig=(AxisKeyName="MotionController_Left_TriggerAxis",AxisProperties=(DeadZone=0.000000,Sensitivity=1.000000,Exponent=1.000000,bInvert=False))
++AxisConfig=(AxisKeyName="MotionController_Left_Grip1Axis",AxisProperties=(DeadZone=0.000000,Sensitivity=1.000000,Exponent=1.000000,bInvert=False))
++AxisConfig=(AxisKeyName="MotionController_Left_Grip2Axis",AxisProperties=(DeadZone=0.000000,Sensitivity=1.000000,Exponent=1.000000,bInvert=False))
++AxisConfig=(AxisKeyName="MotionController_Right_Thumbstick_X",AxisProperties=(DeadZone=0.000000,Sensitivity=1.000000,Exponent=1.000000,bInvert=False))
++AxisConfig=(AxisKeyName="MotionController_Right_Thumbstick_Y",AxisProperties=(DeadZone=0.000000,Sensitivity=1.000000,Exponent=1.000000,bInvert=False))
++AxisConfig=(AxisKeyName="MotionController_Right_TriggerAxis",AxisProperties=(DeadZone=0.000000,Sensitivity=1.000000,Exponent=1.000000,bInvert=False))
++AxisConfig=(AxisKeyName="MotionController_Right_Grip1Axis",AxisProperties=(DeadZone=0.000000,Sensitivity=1.000000,Exponent=1.000000,bInvert=False))
++AxisConfig=(AxisKeyName="MotionController_Right_Grip2Axis",AxisProperties=(DeadZone=0.000000,Sensitivity=1.000000,Exponent=1.000000,bInvert=False))
++AxisConfig=(AxisKeyName="Gamepad_Special_Left_X",AxisProperties=(DeadZone=0.000000,Sensitivity=1.000000,Exponent=1.000000,bInvert=False))
++AxisConfig=(AxisKeyName="Gamepad_Special_Left_Y",AxisProperties=(DeadZone=0.000000,Sensitivity=1.000000,Exponent=1.000000,bInvert=False))
++AxisConfig=(AxisKeyName="Daydream_Left_Trackpad_X",AxisProperties=(DeadZone=0.000000,Sensitivity=1.000000,Exponent=1.000000,bInvert=False))
++AxisConfig=(AxisKeyName="Daydream_Left_Trackpad_Y",AxisProperties=(DeadZone=0.000000,Sensitivity=1.000000,Exponent=1.000000,bInvert=False))
++AxisConfig=(AxisKeyName="Daydream_Right_Trackpad_X",AxisProperties=(DeadZone=0.000000,Sensitivity=1.000000,Exponent=1.000000,bInvert=False))
++AxisConfig=(AxisKeyName="Daydream_Right_Trackpad_Y",AxisProperties=(DeadZone=0.000000,Sensitivity=1.000000,Exponent=1.000000,bInvert=False))
++AxisConfig=(AxisKeyName="Vive_Left_Trigger_Axis",AxisProperties=(DeadZone=0.000000,Sensitivity=1.000000,Exponent=1.000000,bInvert=False))
++AxisConfig=(AxisKeyName="Vive_Left_Trackpad_X",AxisProperties=(DeadZone=0.000000,Sensitivity=1.000000,Exponent=1.000000,bInvert=False))
++AxisConfig=(AxisKeyName="Vive_Left_Trackpad_Y",AxisProperties=(DeadZone=0.000000,Sensitivity=1.000000,Exponent=1.000000,bInvert=False))
++AxisConfig=(AxisKeyName="Vive_Right_Trigger_Axis",AxisProperties=(DeadZone=0.000000,Sensitivity=1.000000,Exponent=1.000000,bInvert=False))
++AxisConfig=(AxisKeyName="Vive_Right_Trackpad_X",AxisProperties=(DeadZone=0.000000,Sensitivity=1.000000,Exponent=1.000000,bInvert=False))
++AxisConfig=(AxisKeyName="Vive_Right_Trackpad_Y",AxisProperties=(DeadZone=0.000000,Sensitivity=1.000000,Exponent=1.000000,bInvert=False))
++AxisConfig=(AxisKeyName="MixedReality_Left_Trigger_Axis",AxisProperties=(DeadZone=0.000000,Sensitivity=1.000000,Exponent=1.000000,bInvert=False))
++AxisConfig=(AxisKeyName="MixedReality_Left_Thumbstick_X",AxisProperties=(DeadZone=0.000000,Sensitivity=1.000000,Exponent=1.000000,bInvert=False))
++AxisConfig=(AxisKeyName="MixedReality_Left_Thumbstick_Y",AxisProperties=(DeadZone=0.000000,Sensitivity=1.000000,Exponent=1.000000,bInvert=False))
++AxisConfig=(AxisKeyName="MixedReality_Left_Trackpad_X",AxisProperties=(DeadZone=0.000000,Sensitivity=1.000000,Exponent=1.000000,bInvert=False))
++AxisConfig=(AxisKeyName="MixedReality_Left_Trackpad_Y",AxisProperties=(DeadZone=0.000000,Sensitivity=1.000000,Exponent=1.000000,bInvert=False))
++AxisConfig=(AxisKeyName="MixedReality_Right_Trigger_Axis",AxisProperties=(DeadZone=0.000000,Sensitivity=1.000000,Exponent=1.000000,bInvert=False))
++AxisConfig=(AxisKeyName="MixedReality_Right_Thumbstick_X",AxisProperties=(DeadZone=0.000000,Sensitivity=1.000000,Exponent=1.000000,bInvert=False))
++AxisConfig=(AxisKeyName="MixedReality_Right_Thumbstick_Y",AxisProperties=(DeadZone=0.000000,Sensitivity=1.000000,Exponent=1.000000,bInvert=False))
++AxisConfig=(AxisKeyName="MixedReality_Right_Trackpad_X",AxisProperties=(DeadZone=0.000000,Sensitivity=1.000000,Exponent=1.000000,bInvert=False))
++AxisConfig=(AxisKeyName="MixedReality_Right_Trackpad_Y",AxisProperties=(DeadZone=0.000000,Sensitivity=1.000000,Exponent=1.000000,bInvert=False))
++AxisConfig=(AxisKeyName="OculusGo_Left_Trackpad_X",AxisProperties=(DeadZone=0.000000,Sensitivity=1.000000,Exponent=1.000000,bInvert=False))
++AxisConfig=(AxisKeyName="OculusGo_Left_Trackpad_Y",AxisProperties=(DeadZone=0.000000,Sensitivity=1.000000,Exponent=1.000000,bInvert=False))
++AxisConfig=(AxisKeyName="OculusGo_Right_Trackpad_X",AxisProperties=(DeadZone=0.000000,Sensitivity=1.000000,Exponent=1.000000,bInvert=False))
++AxisConfig=(AxisKeyName="OculusGo_Right_Trackpad_Y",AxisProperties=(DeadZone=0.000000,Sensitivity=1.000000,Exponent=1.000000,bInvert=False))
++AxisConfig=(AxisKeyName="OculusTouch_Left_Grip_Axis",AxisProperties=(DeadZone=0.000000,Sensitivity=1.000000,Exponent=1.000000,bInvert=False))
++AxisConfig=(AxisKeyName="OculusTouch_Left_Trigger_Axis",AxisProperties=(DeadZone=0.000000,Sensitivity=1.000000,Exponent=1.000000,bInvert=False))
++AxisConfig=(AxisKeyName="OculusTouch_Left_Thumbstick_X",AxisProperties=(DeadZone=0.000000,Sensitivity=1.000000,Exponent=1.000000,bInvert=False))
++AxisConfig=(AxisKeyName="OculusTouch_Left_Thumbstick_Y",AxisProperties=(DeadZone=0.000000,Sensitivity=1.000000,Exponent=1.000000,bInvert=False))
++AxisConfig=(AxisKeyName="OculusTouch_Right_Grip_Axis",AxisProperties=(DeadZone=0.000000,Sensitivity=1.000000,Exponent=1.000000,bInvert=False))
++AxisConfig=(AxisKeyName="OculusTouch_Right_Trigger_Axis",AxisProperties=(DeadZone=0.000000,Sensitivity=1.000000,Exponent=1.000000,bInvert=False))
++AxisConfig=(AxisKeyName="OculusTouch_Right_Thumbstick_X",AxisProperties=(DeadZone=0.000000,Sensitivity=1.000000,Exponent=1.000000,bInvert=False))
++AxisConfig=(AxisKeyName="OculusTouch_Right_Thumbstick_Y",AxisProperties=(DeadZone=0.000000,Sensitivity=1.000000,Exponent=1.000000,bInvert=False))
++AxisConfig=(AxisKeyName="ValveIndex_Left_Grip_Axis",AxisProperties=(DeadZone=0.000000,Sensitivity=1.000000,Exponent=1.000000,bInvert=False))
++AxisConfig=(AxisKeyName="ValveIndex_Left_Grip_Force",AxisProperties=(DeadZone=0.000000,Sensitivity=1.000000,Exponent=1.000000,bInvert=False))
++AxisConfig=(AxisKeyName="ValveIndex_Left_Trigger_Axis",AxisProperties=(DeadZone=0.000000,Sensitivity=1.000000,Exponent=1.000000,bInvert=False))
++AxisConfig=(AxisKeyName="ValveIndex_Left_Thumbstick_X",AxisProperties=(DeadZone=0.000000,Sensitivity=1.000000,Exponent=1.000000,bInvert=False))
++AxisConfig=(AxisKeyName="ValveIndex_Left_Thumbstick_Y",AxisProperties=(DeadZone=0.000000,Sensitivity=1.000000,Exponent=1.000000,bInvert=False))
++AxisConfig=(AxisKeyName="ValveIndex_Left_Trackpad_X",AxisProperties=(DeadZone=0.000000,Sensitivity=1.000000,Exponent=1.000000,bInvert=False))
++AxisConfig=(AxisKeyName="ValveIndex_Left_Trackpad_Y",AxisProperties=(DeadZone=0.000000,Sensitivity=1.000000,Exponent=1.000000,bInvert=False))
++AxisConfig=(AxisKeyName="ValveIndex_Left_Trackpad_Force",AxisProperties=(DeadZone=0.000000,Sensitivity=1.000000,Exponent=1.000000,bInvert=False))
++AxisConfig=(AxisKeyName="ValveIndex_Left_Trackpad_Touch",AxisProperties=(DeadZone=0.000000,Sensitivity=1.000000,Exponent=1.000000,bInvert=False))
++AxisConfig=(AxisKeyName="ValveIndex_Right_Grip_Axis",AxisProperties=(DeadZone=0.000000,Sensitivity=1.000000,Exponent=1.000000,bInvert=False))
++AxisConfig=(AxisKeyName="ValveIndex_Right_Grip_Force",AxisProperties=(DeadZone=0.000000,Sensitivity=1.000000,Exponent=1.000000,bInvert=False))
++AxisConfig=(AxisKeyName="ValveIndex_Right_Trigger_Axis",AxisProperties=(DeadZone=0.000000,Sensitivity=1.000000,Exponent=1.000000,bInvert=False))
++AxisConfig=(AxisKeyName="ValveIndex_Right_Thumbstick_X",AxisProperties=(DeadZone=0.000000,Sensitivity=1.000000,Exponent=1.000000,bInvert=False))
++AxisConfig=(AxisKeyName="ValveIndex_Right_Thumbstick_Y",AxisProperties=(DeadZone=0.000000,Sensitivity=1.000000,Exponent=1.000000,bInvert=False))
++AxisConfig=(AxisKeyName="ValveIndex_Right_Trackpad_X",AxisProperties=(DeadZone=0.000000,Sensitivity=1.000000,Exponent=1.000000,bInvert=False))
++AxisConfig=(AxisKeyName="ValveIndex_Right_Trackpad_Y",AxisProperties=(DeadZone=0.000000,Sensitivity=1.000000,Exponent=1.000000,bInvert=False))
++AxisConfig=(AxisKeyName="ValveIndex_Right_Trackpad_Force",AxisProperties=(DeadZone=0.000000,Sensitivity=1.000000,Exponent=1.000000,bInvert=False))
++AxisConfig=(AxisKeyName="MotionController_Left_Thumbstick_Z",AxisProperties=(DeadZone=0.000000,Sensitivity=1.000000,Exponent=1.000000,bInvert=False))
++AxisConfig=(AxisKeyName="MagicLeap_Left_Trigger_Axis",AxisProperties=(DeadZone=0.000000,Sensitivity=1.000000,Exponent=1.000000,bInvert=False))
++AxisConfig=(AxisKeyName="MagicLeap_Left_Trackpad_X",AxisProperties=(DeadZone=0.000000,Sensitivity=1.000000,Exponent=1.000000,bInvert=False))
++AxisConfig=(AxisKeyName="MagicLeap_Left_Trackpad_Y",AxisProperties=(DeadZone=0.000000,Sensitivity=1.000000,Exponent=1.000000,bInvert=False))
++AxisConfig=(AxisKeyName="MagicLeap_Left_Trackpad_Force",AxisProperties=(DeadZone=0.000000,Sensitivity=1.000000,Exponent=1.000000,bInvert=False))
++AxisConfig=(AxisKeyName="MagicLeap_Left_Touch1_X",AxisProperties=(DeadZone=0.000000,Sensitivity=1.000000,Exponent=1.000000,bInvert=False))
++AxisConfig=(AxisKeyName="MagicLeap_Left_Touch1_Y",AxisProperties=(DeadZone=0.000000,Sensitivity=1.000000,Exponent=1.000000,bInvert=False))
++AxisConfig=(AxisKeyName="MagicLeap_Left_Touch1_Force",AxisProperties=(DeadZone=0.000000,Sensitivity=1.000000,Exponent=1.000000,bInvert=False))
++AxisConfig=(AxisKeyName="MotionController_Right_Thumbstick_Z",AxisProperties=(DeadZone=0.000000,Sensitivity=1.000000,Exponent=1.000000,bInvert=False))
++AxisConfig=(AxisKeyName="MagicLeap_Right_Trigger_Axis",AxisProperties=(DeadZone=0.000000,Sensitivity=1.000000,Exponent=1.000000,bInvert=False))
++AxisConfig=(AxisKeyName="MagicLeap_Right_Trackpad_X",AxisProperties=(DeadZone=0.000000,Sensitivity=1.000000,Exponent=1.000000,bInvert=False))
++AxisConfig=(AxisKeyName="MagicLeap_Right_Trackpad_Y",AxisProperties=(DeadZone=0.000000,Sensitivity=1.000000,Exponent=1.000000,bInvert=False))
++AxisConfig=(AxisKeyName="MagicLeap_Right_Trackpad_Force",AxisProperties=(DeadZone=0.000000,Sensitivity=1.000000,Exponent=1.000000,bInvert=False))
++AxisConfig=(AxisKeyName="MagicLeap_Right_Touch1_X",AxisProperties=(DeadZone=0.000000,Sensitivity=1.000000,Exponent=1.000000,bInvert=False))
++AxisConfig=(AxisKeyName="MagicLeap_Right_Touch1_Y",AxisProperties=(DeadZone=0.000000,Sensitivity=1.000000,Exponent=1.000000,bInvert=False))
++AxisConfig=(AxisKeyName="MagicLeap_Right_Touch1_Force",AxisProperties=(DeadZone=0.000000,Sensitivity=1.000000,Exponent=1.000000,bInvert=False))
++AxisConfig=(AxisKeyName="OculusTouch_Left_Thumbstick",AxisProperties=(DeadZone=0.000000,Sensitivity=1.000000,Exponent=1.000000,bInvert=False))
++AxisConfig=(AxisKeyName="OculusTouch_Left_FaceButton1",AxisProperties=(DeadZone=0.000000,Sensitivity=1.000000,Exponent=1.000000,bInvert=False))
++AxisConfig=(AxisKeyName="OculusTouch_Left_Trigger",AxisProperties=(DeadZone=0.000000,Sensitivity=1.000000,Exponent=1.000000,bInvert=False))
++AxisConfig=(AxisKeyName="OculusTouch_Left_FaceButton2",AxisProperties=(DeadZone=0.000000,Sensitivity=1.000000,Exponent=1.000000,bInvert=False))
++AxisConfig=(AxisKeyName="OculusTouch_Left_IndexPointing",AxisProperties=(DeadZone=0.000000,Sensitivity=1.000000,Exponent=1.000000,bInvert=False))
++AxisConfig=(AxisKeyName="OculusTouch_Left_ThumbUp",AxisProperties=(DeadZone=0.000000,Sensitivity=1.000000,Exponent=1.000000,bInvert=False))
++AxisConfig=(AxisKeyName="OculusTouch_Right_Thumbstick",AxisProperties=(DeadZone=0.000000,Sensitivity=1.000000,Exponent=1.000000,bInvert=False))
++AxisConfig=(AxisKeyName="OculusTouch_Right_FaceButton1",AxisProperties=(DeadZone=0.000000,Sensitivity=1.000000,Exponent=1.000000,bInvert=False))
++AxisConfig=(AxisKeyName="OculusTouch_Right_Trigger",AxisProperties=(DeadZone=0.000000,Sensitivity=1.000000,Exponent=1.000000,bInvert=False))
++AxisConfig=(AxisKeyName="OculusTouch_Right_FaceButton2",AxisProperties=(DeadZone=0.000000,Sensitivity=1.000000,Exponent=1.000000,bInvert=False))
++AxisConfig=(AxisKeyName="OculusTouch_Right_IndexPointing",AxisProperties=(DeadZone=0.000000,Sensitivity=1.000000,Exponent=1.000000,bInvert=False))
++AxisConfig=(AxisKeyName="OculusTouch_Right_ThumbUp",AxisProperties=(DeadZone=0.000000,Sensitivity=1.000000,Exponent=1.000000,bInvert=False))
++AxisConfig=(AxisKeyName="OculusTouchpad_Touchpad_X",AxisProperties=(DeadZone=0.000000,Sensitivity=1.000000,Exponent=1.000000,bInvert=False))
++AxisConfig=(AxisKeyName="OculusTouchpad_Touchpad_Y",AxisProperties=(DeadZone=0.000000,Sensitivity=1.000000,Exponent=1.000000,bInvert=False))
 bAltEnterTogglesFullscreen=True
+bF11TogglesFullscreen=True
 bUseMouseForTouch=False
 bEnableMouseSmoothing=True
 bEnableFOVScaling=True
-FOVScale=0.011110
-DoubleClickTime=0.200000
-+ActionMappings=(ActionName="Handbrake",Key=SpaceBar,bShift=False,bCtrl=False,bAlt=False,bCmd=False)
-+ActionMappings=(ActionName="Handbrake",Key=Gamepad_LeftShoulder,bShift=False,bCtrl=False,bAlt=False,bCmd=False)
-+ActionMappings=(ActionName="SwitchCamera",Key=Tab,bShift=False,bCtrl=False,bAlt=False,bCmd=False)
-+ActionMappings=(ActionName="ResetVR",Key=R,bShift=False,bCtrl=False,bAlt=False,bCmd=False)
-+ActionMappings=(ActionName="ResetVR",Key=MotionController_Left_Grip1,bShift=False,bCtrl=False,bAlt=False,bCmd=False)
-+ActionMappings=(ActionName="SwitchCamera",Key=Gamepad_Special_Left,bShift=False,bCtrl=False,bAlt=False,bCmd=False)
-+AxisMappings=(AxisName="MoveForward",Key=W,Scale=1.000000)
-+AxisMappings=(AxisName="MoveForward",Key=S,Scale=-1.000000)
-+AxisMappings=(AxisName="MoveForward",Key=Up,Scale=1.000000)
-+AxisMappings=(AxisName="MoveForward",Key=Down,Scale=-1.000000)
-+AxisMappings=(AxisName="MoveForward",Key=Gamepad_FaceButton_Bottom,Scale=1.000000)
-+AxisMappings=(AxisName="MoveForward",Key=Gamepad_RightY,Scale=-1.000000)
-+AxisMappings=(AxisName="MoveRight",Key=A,Scale=-1.000000)
-+AxisMappings=(AxisName="MoveRight",Key=D,Scale=1.000000)
-+AxisMappings=(AxisName="MoveRight",Key=Left,Scale=-1.000000)
-+AxisMappings=(AxisName="MoveRight",Key=Right,Scale=1.000000)
-+AxisMappings=(AxisName="MoveRight",Key=Gamepad_LeftX,Scale=1.000000)
-+AxisMappings=(AxisName="LookUp",Key=MouseY,Scale=1.000000)
-+AxisMappings=(AxisName="LookRight",Key=MouseX,Scale=1.000000)
-+AxisMappings=(AxisName="MoveForward",Key=Gamepad_RightTriggerAxis,Scale=1.000000)
-+AxisMappings=(AxisName="MoveForward",Key=Gamepad_LeftTriggerAxis,Scale=-1.000000)
-+AxisMappings=(AxisName="MoveForward",Key=GenericUSBController_Axis2,Scale=-2.000000)
-+AxisMappings=(AxisName="MoveForward",Key=GenericUSBController_Axis3,Scale=1.000000)
-+AxisMappings=(AxisName="MoveRight",Key=GenericUSBController_Axis1,Scale=3.000000)
+bCaptureMouseOnLaunch=True
 bAlwaysShowTouchInterface=False
 bShowConsoleOnFourFingerTap=True
+bEnableGestureRecognizer=False
+bUseAutocorrect=False
+DefaultViewportMouseCaptureMode=CapturePermanently_IncludingInitialMouseDown
+DefaultViewportMouseLockMode=LockOnCapture
+FOVScale=0.011110
+DoubleClickTime=0.200000
++ActionMappings=(ActionName="Handbrake",bShift=False,bCtrl=False,bAlt=False,bCmd=False,Key=SpaceBar)
++ActionMappings=(ActionName="Handbrake",bShift=False,bCtrl=False,bAlt=False,bCmd=False,Key=Gamepad_LeftShoulder)
++ActionMappings=(ActionName="SwitchCamera",bShift=False,bCtrl=False,bAlt=False,bCmd=False,Key=Tab)
++ActionMappings=(ActionName="ResetVR",bShift=False,bCtrl=False,bAlt=False,bCmd=False,Key=R)
++ActionMappings=(ActionName="SwitchCamera",bShift=False,bCtrl=False,bAlt=False,bCmd=False,Key=Gamepad_Special_Left)
++ActionMappings=(ActionName="ResetVR",bShift=False,bCtrl=False,bAlt=False,bCmd=False,Key=Vive_Left_Menu_Click)
++ActionMappings=(ActionName="ResetVR",bShift=False,bCtrl=False,bAlt=False,bCmd=False,Key=MixedReality_Left_Thumbstick_Click)
++ActionMappings=(ActionName="ResetVR",bShift=False,bCtrl=False,bAlt=False,bCmd=False,Key=OculusTouch_Left_Thumbstick_Click)
++ActionMappings=(ActionName="ResetVR",bShift=False,bCtrl=False,bAlt=False,bCmd=False,Key=ValveIndex_Left_Thumbstick_Click)
++ActionMappings=(ActionName="ResetVR",bShift=False,bCtrl=False,bAlt=False,bCmd=False,Key=MagicLeap_Left_Bumper)
++ActionMappings=(ActionName="Handbrake",bShift=False,bCtrl=False,bAlt=False,bCmd=False,Key=Vive_Right_Grip_Click)
++ActionMappings=(ActionName="Handbrake",bShift=False,bCtrl=False,bAlt=False,bCmd=False,Key=MixedReality_Right_Grip_Click)
++ActionMappings=(ActionName="Handbrake",bShift=False,bCtrl=False,bAlt=False,bCmd=False,Key=OculusTouch_Right_Grip_Click)
++ActionMappings=(ActionName="Handbrake",bShift=False,bCtrl=False,bAlt=False,bCmd=False,Key=ValveIndex_Right_A_Click)
++ActionMappings=(ActionName="Handbrake",bShift=False,bCtrl=False,bAlt=False,bCmd=False,Key=MagicLeap_Right_Bumper)
++AxisMappings=(AxisName="MoveForward",Scale=1.000000,Key=W)
++AxisMappings=(AxisName="MoveForward",Scale=-1.000000,Key=S)
++AxisMappings=(AxisName="MoveForward",Scale=1.000000,Key=Up)
++AxisMappings=(AxisName="MoveForward",Scale=-1.000000,Key=Down)
++AxisMappings=(AxisName="MoveForward",Scale=1.000000,Key=Gamepad_FaceButton_Bottom)
++AxisMappings=(AxisName="MoveForward",Scale=-1.000000,Key=Gamepad_RightY)
++AxisMappings=(AxisName="MoveRight",Scale=-1.000000,Key=A)
++AxisMappings=(AxisName="MoveRight",Scale=1.000000,Key=D)
++AxisMappings=(AxisName="MoveRight",Scale=-1.000000,Key=Left)
++AxisMappings=(AxisName="MoveRight",Scale=1.000000,Key=Right)
++AxisMappings=(AxisName="MoveRight",Scale=1.000000,Key=Gamepad_LeftX)
++AxisMappings=(AxisName="LookUp",Scale=1.000000,Key=MouseY)
++AxisMappings=(AxisName="LookRight",Scale=1.000000,Key=MouseX)
++AxisMappings=(AxisName="MoveForward",Scale=1.000000,Key=Gamepad_RightTriggerAxis)
++AxisMappings=(AxisName="MoveForward",Scale=-1.000000,Key=Gamepad_LeftTriggerAxis)
++AxisMappings=(AxisName="MoveForward",Scale=-2.000000,Key=GenericUSBController_Axis2)
++AxisMappings=(AxisName="MoveForward",Scale=1.000000,Key=GenericUSBController_Axis3)
++AxisMappings=(AxisName="MoveRight",Scale=3.000000,Key=GenericUSBController_Axis1)
++AxisMappings=(AxisName="MoveForward",Scale=1.000000,Key=Vive_Right_Trigger_Axis)
++AxisMappings=(AxisName="MoveForward",Scale=-1.000000,Key=Vive_Left_Trigger_Axis)
++AxisMappings=(AxisName="MoveForward",Scale=1.000000,Key=MixedReality_Right_Trigger_Axis)
++AxisMappings=(AxisName="MoveForward",Scale=-1.000000,Key=MixedReality_Left_Trigger_Axis)
++AxisMappings=(AxisName="MoveForward",Scale=1.000000,Key=OculusTouch_Right_Trigger_Axis)
++AxisMappings=(AxisName="MoveForward",Scale=-1.000000,Key=OculusTouch_Left_Trigger_Axis)
++AxisMappings=(AxisName="MoveForward",Scale=1.000000,Key=ValveIndex_Right_Trigger_Axis)
++AxisMappings=(AxisName="MoveForward",Scale=-1.000000,Key=ValveIndex_Left_Trigger_Axis)
++AxisMappings=(AxisName="MoveForward",Scale=1.000000,Key=MagicLeap_Right_Trigger_Axis)
++AxisMappings=(AxisName="MoveForward",Scale=-1.000000,Key=MagicLeap_Left_Trigger_Axis)
++AxisMappings=(AxisName="MoveRight",Scale=1.000000,Key=Vive_Left_Trackpad_X)
++AxisMappings=(AxisName="MoveRight",Scale=1.000000,Key=Vive_Right_Trackpad_X)
++AxisMappings=(AxisName="MoveRight",Scale=1.000000,Key=MixedReality_Left_Thumbstick_X)
++AxisMappings=(AxisName="MoveRight",Scale=1.000000,Key=MixedReality_Right_Thumbstick_X)
++AxisMappings=(AxisName="MoveRight",Scale=1.000000,Key=OculusTouch_Left_Thumbstick_X)
++AxisMappings=(AxisName="MoveRight",Scale=1.000000,Key=OculusTouch_Right_Thumbstick_X)
++AxisMappings=(AxisName="MoveRight",Scale=1.000000,Key=ValveIndex_Left_Thumbstick_X)
++AxisMappings=(AxisName="MoveRight",Scale=1.000000,Key=ValveIndex_Right_Thumbstick_X)
++AxisMappings=(AxisName="MoveRight",Scale=1.000000,Key=MagicLeap_Left_Trackpad_X)
++AxisMappings=(AxisName="MoveRight",Scale=1.000000,Key=MagicLeap_Right_Trackpad_X)
 DefaultTouchInterface=/Game/Vehicle/Sedan/SedanVirtualJoysticks.SedanVirtualJoysticks
-ConsoleKey=None
--ConsoleKeys=Tilde
 +ConsoleKeys=Tilde
 
-[/Script/RawInput.RawInputSettings]
-+DeviceConfigurations=(VendorID="0x046D",ProductID="0xC262",AxisProperties=((bEnabled=False),(Key=GenericUSBController_Axis1,Offset=-0.500000),(bEnabled=False),(Key=GenericUSBController_Axis2,bInverted=True,Offset=1.000000),(Key=GenericUSBController_Axis3,bInverted=True,Offset=1.000000)),ButtonProperties=((Key=GenericUSBController_Button1),(Key=GenericUSBController_Button2),(Key=GenericUSBController_Button3),(Key=GenericUSBController_Button4),(Key=GenericUSBController_Button5),(Key=GenericUSBController_Button6),(Key=GenericUSBController_Button7),(Key=GenericUSBController_Button8),(Key=GenericUSBController_Button9),(Key=GenericUSBController_Button10),(Key=GenericUSBController_Button11),(Key=GenericUSBController_Button12)))
-
 
diff --git a/Templates/TP_VehicleAdv/Config/DefaultInput.ini b/Templates/TP_VehicleAdv/Config/DefaultInput.ini
--- a/Templates/TP_VehicleAdv/Config/DefaultInput.ini
+++ b/Templates/TP_VehicleAdv/Config/DefaultInput.ini
@@ -1,3 +1,6 @@
+[/Script/RawInput.RawInputSettings]
++DeviceConfigurations=(VendorID="0x046D",ProductID="0xC262",AxisProperties=((bEnabled=False),(Key=GenericUSBController_Axis1,Offset=-0.500000),(bEnabled=False),(Key=GenericUSBController_Axis2,bInverted=True,Offset=1.000000),(Key=GenericUSBController_Axis3,bInverted=True,Offset=1.000000)),ButtonProperties=((Key=GenericUSBController_Button1),(Key=GenericUSBController_Button2),(Key=GenericUSBController_Button3),(Key=GenericUSBController_Button4),(Key=GenericUSBController_Button5),(Key=GenericUSBController_Button6),(Key=GenericUSBController_Button7),(Key=GenericUSBController_Button8),(Key=GenericUSBController_Button9),(Key=GenericUSBController_Button10),(Key=GenericUSBController_Button11),(Key=GenericUSBController_Button12)))
+
 [/Script/Engine.InputSettings]
 +AxisConfig=(AxisKeyName="Gamepad_LeftX",AxisProperties=(DeadZone=0.250000,Sensitivity=1.000000,Exponent=1.000000,bInvert=False))
 +AxisConfig=(AxisKeyName="Gamepad_LeftY",AxisProperties=(DeadZone=0.250000,Sensitivity=1.000000,Exponent=1.000000,bInvert=False))
@@ -5,38 +8,168 @@
 +AxisConfig=(AxisKeyName="Gamepad_RightY",AxisProperties=(DeadZone=0.250000,Sensitivity=1.000000,Exponent=1.000000,bInvert=False))
 +AxisConfig=(AxisKeyName="MouseX",AxisProperties=(DeadZone=0.000000,Sensitivity=0.070000,Exponent=1.000000,bInvert=False))
 +AxisConfig=(AxisKeyName="MouseY",AxisProperties=(DeadZone=0.000000,Sensitivity=0.070000,Exponent=1.000000,bInvert=False))
++AxisConfig=(AxisKeyName="MouseWheelAxis",AxisProperties=(DeadZone=0.000000,Sensitivity=1.000000,Exponent=1.000000,bInvert=False))
++AxisConfig=(AxisKeyName="Gamepad_LeftTriggerAxis",AxisProperties=(DeadZone=0.000000,Sensitivity=1.000000,Exponent=1.000000,bInvert=False))
++AxisConfig=(AxisKeyName="Gamepad_RightTriggerAxis",AxisProperties=(DeadZone=0.000000,Sensitivity=1.000000,Exponent=1.000000,bInvert=False))
++AxisConfig=(AxisKeyName="MotionController_Left_Thumbstick_X",AxisProperties=(DeadZone=0.000000,Sensitivity=1.000000,Exponent=1.000000,bInvert=False))
++AxisConfig=(AxisKeyName="MotionController_Left_Thumbstick_Y",AxisProperties=(DeadZone=0.000000,Sensitivity=1.000000,Exponent=1.000000,bInvert=False))
++AxisConfig=(AxisKeyName="MotionController_Left_TriggerAxis",AxisProperties=(DeadZone=0.000000,Sensitivity=1.000000,Exponent=1.000000,bInvert=False))
++AxisConfig=(AxisKeyName="MotionController_Left_Grip1Axis",AxisProperties=(DeadZone=0.000000,Sensitivity=1.000000,Exponent=1.000000,bInvert=False))
++AxisConfig=(AxisKeyName="MotionController_Left_Grip2Axis",AxisProperties=(DeadZone=0.000000,Sensitivity=1.000000,Exponent=1.000000,bInvert=False))
++AxisConfig=(AxisKeyName="MotionController_Right_Thumbstick_X",AxisProperties=(DeadZone=0.000000,Sensitivity=1.000000,Exponent=1.000000,bInvert=False))
++AxisConfig=(AxisKeyName="MotionController_Right_Thumbstick_Y",AxisProperties=(DeadZone=0.000000,Sensitivity=1.000000,Exponent=1.000000,bInvert=False))
++AxisConfig=(AxisKeyName="MotionController_Right_TriggerAxis",AxisProperties=(DeadZone=0.000000,Sensitivity=1.000000,Exponent=1.000000,bInvert=False))
++AxisConfig=(AxisKeyName="MotionController_Right_Grip1Axis",AxisProperties=(DeadZone=0.000000,Sensitivity=1.000000,Exponent=1.000000,bInvert=False))
++AxisConfig=(AxisKeyName="MotionController_Right_Grip2Axis",AxisProperties=(DeadZone=0.000000,Sensitivity=1.000000,Exponent=1.000000,bInvert=False))
++AxisConfig=(AxisKeyName="Gamepad_Special_Left_X",AxisProperties=(DeadZone=0.000000,Sensitivity=1.000000,Exponent=1.000000,bInvert=False))
++AxisConfig=(AxisKeyName="Gamepad_Special_Left_Y",AxisProperties=(DeadZone=0.000000,Sensitivity=1.000000,Exponent=1.000000,bInvert=False))
++AxisConfig=(AxisKeyName="Daydream_Left_Trackpad_X",AxisProperties=(DeadZone=0.000000,Sensitivity=1.000000,Exponent=1.000000,bInvert=False))
++AxisConfig=(AxisKeyName="Daydream_Left_Trackpad_Y",AxisProperties=(DeadZone=0.000000,Sensitivity=1.000000,Exponent=1.000000,bInvert=False))
++AxisConfig=(AxisKeyName="Daydream_Right_Trackpad_X",AxisProperties=(DeadZone=0.000000,Sensitivity=1.000000,Exponent=1.000000,bInvert=False))
++AxisConfig=(AxisKeyName="Daydream_Right_Trackpad_Y",AxisProperties=(DeadZone=0.000000,Sensitivity=1.000000,Exponent=1.000000,bInvert=False))
++AxisConfig=(AxisKeyName="Vive_Left_Trigger_Axis",AxisProperties=(DeadZone=0.000000,Sensitivity=1.000000,Exponent=1.000000,bInvert=False))
++AxisConfig=(AxisKeyName="Vive_Left_Trackpad_X",AxisProperties=(DeadZone=0.000000,Sensitivity=1.000000,Exponent=1.000000,bInvert=False))
++AxisConfig=(AxisKeyName="Vive_Left_Trackpad_Y",AxisProperties=(DeadZone=0.000000,Sensitivity=1.000000,Exponent=1.000000,bInvert=False))
++AxisConfig=(AxisKeyName="Vive_Right_Trigger_Axis",AxisProperties=(DeadZone=0.000000,Sensitivity=1.000000,Exponent=1.000000,bInvert=False))
++AxisConfig=(AxisKeyName="Vive_Right_Trackpad_X",AxisProperties=(DeadZone=0.000000,Sensitivity=1.000000,Exponent=1.000000,bInvert=False))
++AxisConfig=(AxisKeyName="Vive_Right_Trackpad_Y",AxisProperties=(DeadZone=0.000000,Sensitivity=1.000000,Exponent=1.000000,bInvert=False))
++AxisConfig=(AxisKeyName="MixedReality_Left_Trigger_Axis",AxisProperties=(DeadZone=0.000000,Sensitivity=1.000000,Exponent=1.000000,bInvert=False))
++AxisConfig=(AxisKeyName="MixedReality_Left_Thumbstick_X",AxisProperties=(DeadZone=0.000000,Sensitivity=1.000000,Exponent=1.000000,bInvert=False))
++AxisConfig=(AxisKeyName="MixedReality_Left_Thumbstick_Y",AxisProperties=(DeadZone=0.000000,Sensitivity=1.000000,Exponent=1.000000,bInvert=False))
++AxisConfig=(AxisKeyName="MixedReality_Left_Trackpad_X",AxisProperties=(DeadZone=0.000000,Sensitivity=1.000000,Exponent=1.000000,bInvert=False))
++AxisConfig=(AxisKeyName="MixedReality_Left_Trackpad_Y",AxisProperties=(DeadZone=0.000000,Sensitivity=1.000000,Exponent=1.000000,bInvert=False))
++AxisConfig=(AxisKeyName="MixedReality_Right_Trigger_Axis",AxisProperties=(DeadZone=0.000000,Sensitivity=1.000000,Exponent=1.000000,bInvert=False))
++AxisConfig=(AxisKeyName="MixedReality_Right_Thumbstick_X",AxisProperties=(DeadZone=0.000000,Sensitivity=1.000000,Exponent=1.000000,bInvert=False))
++AxisConfig=(AxisKeyName="MixedReality_Right_Thumbstick_Y",AxisProperties=(DeadZone=0.000000,Sensitivity=1.000000,Exponent=1.000000,bInvert=False))
++AxisConfig=(AxisKeyName="MixedReality_Right_Trackpad_X",AxisProperties=(DeadZone=0.000000,Sensitivity=1.000000,Exponent=1.000000,bInvert=False))
++AxisConfig=(AxisKeyName="MixedReality_Right_Trackpad_Y",AxisProperties=(DeadZone=0.000000,Sensitivity=1.000000,Exponent=1.000000,bInvert=False))
++AxisConfig=(AxisKeyName="OculusGo_Left_Trackpad_X",AxisProperties=(DeadZone=0.000000,Sensitivity=1.000000,Exponent=1.000000,bInvert=False))
++AxisConfig=(AxisKeyName="OculusGo_Left_Trackpad_Y",AxisProperties=(DeadZone=0.000000,Sensitivity=1.000000,Exponent=1.000000,bInvert=False))
++AxisConfig=(AxisKeyName="OculusGo_Right_Trackpad_X",AxisProperties=(DeadZone=0.000000,Sensitivity=1.000000,Exponent=1.000000,bInvert=False))
++AxisConfig=(AxisKeyName="OculusGo_Right_Trackpad_Y",AxisProperties=(DeadZone=0.000000,Sensitivity=1.000000,Exponent=1.000000,bInvert=False))
++AxisConfig=(AxisKeyName="OculusTouch_Left_Grip_Axis",AxisProperties=(DeadZone=0.000000,Sensitivity=1.000000,Exponent=1.000000,bInvert=False))
++AxisConfig=(AxisKeyName="OculusTouch_Left_Trigger_Axis",AxisProperties=(DeadZone=0.000000,Sensitivity=1.000000,Exponent=1.000000,bInvert=False))
++AxisConfig=(AxisKeyName="OculusTouch_Left_Thumbstick_X",AxisProperties=(DeadZone=0.000000,Sensitivity=1.000000,Exponent=1.000000,bInvert=False))
++AxisConfig=(AxisKeyName="OculusTouch_Left_Thumbstick_Y",AxisProperties=(DeadZone=0.000000,Sensitivity=1.000000,Exponent=1.000000,bInvert=False))
++AxisConfig=(AxisKeyName="OculusTouch_Right_Grip_Axis",AxisProperties=(DeadZone=0.000000,Sensitivity=1.000000,Exponent=1.000000,bInvert=False))
++AxisConfig=(AxisKeyName="OculusTouch_Right_Trigger_Axis",AxisProperties=(DeadZone=0.000000,Sensitivity=1.000000,Exponent=1.000000,bInvert=False))
++AxisConfig=(AxisKeyName="OculusTouch_Right_Thumbstick_X",AxisProperties=(DeadZone=0.000000,Sensitivity=1.000000,Exponent=1.000000,bInvert=False))
++AxisConfig=(AxisKeyName="OculusTouch_Right_Thumbstick_Y",AxisProperties=(DeadZone=0.000000,Sensitivity=1.000000,Exponent=1.000000,bInvert=False))
++AxisConfig=(AxisKeyName="ValveIndex_Left_Grip_Axis",AxisProperties=(DeadZone=0.000000,Sensitivity=1.000000,Exponent=1.000000,bInvert=False))
++AxisConfig=(AxisKeyName="ValveIndex_Left_Grip_Force",AxisProperties=(DeadZone=0.000000,Sensitivity=1.000000,Exponent=1.000000,bInvert=False))
++AxisConfig=(AxisKeyName="ValveIndex_Left_Trigger_Axis",AxisProperties=(DeadZone=0.000000,Sensitivity=1.000000,Exponent=1.000000,bInvert=False))
++AxisConfig=(AxisKeyName="ValveIndex_Left_Thumbstick_X",AxisProperties=(DeadZone=0.000000,Sensitivity=1.000000,Exponent=1.000000,bInvert=False))
++AxisConfig=(AxisKeyName="ValveIndex_Left_Thumbstick_Y",AxisProperties=(DeadZone=0.000000,Sensitivity=1.000000,Exponent=1.000000,bInvert=False))
++AxisConfig=(AxisKeyName="ValveIndex_Left_Trackpad_X",AxisProperties=(DeadZone=0.000000,Sensitivity=1.000000,Exponent=1.000000,bInvert=False))
++AxisConfig=(AxisKeyName="ValveIndex_Left_Trackpad_Y",AxisProperties=(DeadZone=0.000000,Sensitivity=1.000000,Exponent=1.000000,bInvert=False))
++AxisConfig=(AxisKeyName="ValveIndex_Left_Trackpad_Force",AxisProperties=(DeadZone=0.000000,Sensitivity=1.000000,Exponent=1.000000,bInvert=False))
++AxisConfig=(AxisKeyName="ValveIndex_Left_Trackpad_Touch",AxisProperties=(DeadZone=0.000000,Sensitivity=1.000000,Exponent=1.000000,bInvert=False))
++AxisConfig=(AxisKeyName="ValveIndex_Right_Grip_Axis",AxisProperties=(DeadZone=0.000000,Sensitivity=1.000000,Exponent=1.000000,bInvert=False))
++AxisConfig=(AxisKeyName="ValveIndex_Right_Grip_Force",AxisProperties=(DeadZone=0.000000,Sensitivity=1.000000,Exponent=1.000000,bInvert=False))
++AxisConfig=(AxisKeyName="ValveIndex_Right_Trigger_Axis",AxisProperties=(DeadZone=0.000000,Sensitivity=1.000000,Exponent=1.000000,bInvert=False))
++AxisConfig=(AxisKeyName="ValveIndex_Right_Thumbstick_X",AxisProperties=(DeadZone=0.000000,Sensitivity=1.000000,Exponent=1.000000,bInvert=False))
++AxisConfig=(AxisKeyName="ValveIndex_Right_Thumbstick_Y",AxisProperties=(DeadZone=0.000000,Sensitivity=1.000000,Exponent=1.000000,bInvert=False))
++AxisConfig=(AxisKeyName="ValveIndex_Right_Trackpad_X",AxisProperties=(DeadZone=0.000000,Sensitivity=1.000000,Exponent=1.000000,bInvert=False))
++AxisConfig=(AxisKeyName="ValveIndex_Right_Trackpad_Y",AxisProperties=(DeadZone=0.000000,Sensitivity=1.000000,Exponent=1.000000,bInvert=False))
++AxisConfig=(AxisKeyName="ValveIndex_Right_Trackpad_Force",AxisProperties=(DeadZone=0.000000,Sensitivity=1.000000,Exponent=1.000000,bInvert=False))
++AxisConfig=(AxisKeyName="MotionController_Left_Thumbstick_Z",AxisProperties=(DeadZone=0.000000,Sensitivity=1.000000,Exponent=1.000000,bInvert=False))
++AxisConfig=(AxisKeyName="MagicLeap_Left_Trigger_Axis",AxisProperties=(DeadZone=0.000000,Sensitivity=1.000000,Exponent=1.000000,bInvert=False))
++AxisConfig=(AxisKeyName="MagicLeap_Left_Trackpad_X",AxisProperties=(DeadZone=0.000000,Sensitivity=1.000000,Exponent=1.000000,bInvert=False))
++AxisConfig=(AxisKeyName="MagicLeap_Left_Trackpad_Y",AxisProperties=(DeadZone=0.000000,Sensitivity=1.000000,Exponent=1.000000,bInvert=False))
++AxisConfig=(AxisKeyName="MagicLeap_Left_Trackpad_Force",AxisProperties=(DeadZone=0.000000,Sensitivity=1.000000,Exponent=1.000000,bInvert=False))
++AxisConfig=(AxisKeyName="MagicLeap_Left_Touch1_X",AxisProperties=(DeadZone=0.000000,Sensitivity=1.000000,Exponent=1.000000,bInvert=False))
++AxisConfig=(AxisKeyName="MagicLeap_Left_Touch1_Y",AxisProperties=(DeadZone=0.000000,Sensitivity=1.000000,Exponent=1.000000,bInvert=False))
++AxisConfig=(AxisKeyName="MagicLeap_Left_Touch1_Force",AxisProperties=(DeadZone=0.000000,Sensitivity=1.000000,Exponent=1.000000,bInvert=False))
++AxisConfig=(AxisKeyName="MotionController_Right_Thumbstick_Z",AxisProperties=(DeadZone=0.000000,Sensitivity=1.000000,Exponent=1.000000,bInvert=False))
++AxisConfig=(AxisKeyName="MagicLeap_Right_Trigger_Axis",AxisProperties=(DeadZone=0.000000,Sensitivity=1.000000,Exponent=1.000000,bInvert=False))
++AxisConfig=(AxisKeyName="MagicLeap_Right_Trackpad_X",AxisProperties=(DeadZone=0.000000,Sensitivity=1.000000,Exponent=1.000000,bInvert=False))
++AxisConfig=(AxisKeyName="MagicLeap_Right_Trackpad_Y",AxisProperties=(DeadZone=0.000000,Sensitivity=1.000000,Exponent=1.000000,bInvert=False))
++AxisConfig=(AxisKeyName="MagicLeap_Right_Trackpad_Force",AxisProperties=(DeadZone=0.000000,Sensitivity=1.000000,Exponent=1.000000,bInvert=False))
++AxisConfig=(AxisKeyName="MagicLeap_Right_Touch1_X",AxisProperties=(DeadZone=0.000000,Sensitivity=1.000000,Exponent=1.000000,bInvert=False))
++AxisConfig=(AxisKeyName="MagicLeap_Right_Touch1_Y",AxisProperties=(DeadZone=0.000000,Sensitivity=1.000000,Exponent=1.000000,bInvert=False))
++AxisConfig=(AxisKeyName="MagicLeap_Right_Touch1_Force",AxisProperties=(DeadZone=0.000000,Sensitivity=1.000000,Exponent=1.000000,bInvert=False))
++AxisConfig=(AxisKeyName="OculusTouch_Left_Thumbstick",AxisProperties=(DeadZone=0.000000,Sensitivity=1.000000,Exponent=1.000000,bInvert=False))
++AxisConfig=(AxisKeyName="OculusTouch_Left_FaceButton1",AxisProperties=(DeadZone=0.000000,Sensitivity=1.000000,Exponent=1.000000,bInvert=False))
++AxisConfig=(AxisKeyName="OculusTouch_Left_Trigger",AxisProperties=(DeadZone=0.000000,Sensitivity=1.000000,Exponent=1.000000,bInvert=False))
++AxisConfig=(AxisKeyName="OculusTouch_Left_FaceButton2",AxisProperties=(DeadZone=0.000000,Sensitivity=1.000000,Exponent=1.000000,bInvert=False))
++AxisConfig=(AxisKeyName="OculusTouch_Left_IndexPointing",AxisProperties=(DeadZone=0.000000,Sensitivity=1.000000,Exponent=1.000000,bInvert=False))
++AxisConfig=(AxisKeyName="OculusTouch_Left_ThumbUp",AxisProperties=(DeadZone=0.000000,Sensitivity=1.000000,Exponent=1.000000,bInvert=False))
++AxisConfig=(AxisKeyName="OculusTouch_Right_Thumbstick",AxisProperties=(DeadZone=0.000000,Sensitivity=1.000000,Exponent=1.000000,bInvert=False))
++AxisConfig=(AxisKeyName="OculusTouch_Right_FaceButton1",AxisProperties=(DeadZone=0.000000,Sensitivity=1.000000,Exponent=1.000000,bInvert=False))
++AxisConfig=(AxisKeyName="OculusTouch_Right_Trigger",AxisProperties=(DeadZone=0.000000,Sensitivity=1.000000,Exponent=1.000000,bInvert=False))
++AxisConfig=(AxisKeyName="OculusTouch_Right_FaceButton2",AxisProperties=(DeadZone=0.000000,Sensitivity=1.000000,Exponent=1.000000,bInvert=False))
++AxisConfig=(AxisKeyName="OculusTouch_Right_IndexPointing",AxisProperties=(DeadZone=0.000000,Sensitivity=1.000000,Exponent=1.000000,bInvert=False))
++AxisConfig=(AxisKeyName="OculusTouch_Right_ThumbUp",AxisProperties=(DeadZone=0.000000,Sensitivity=1.000000,Exponent=1.000000,bInvert=False))
++AxisConfig=(AxisKeyName="OculusTouchpad_Touchpad_X",AxisProperties=(DeadZone=0.000000,Sensitivity=1.000000,Exponent=1.000000,bInvert=False))
++AxisConfig=(AxisKeyName="OculusTouchpad_Touchpad_Y",AxisProperties=(DeadZone=0.000000,Sensitivity=1.000000,Exponent=1.000000,bInvert=False))
+bAltEnterTogglesFullscreen=True
+bF11TogglesFullscreen=True
 bUseMouseForTouch=False
-+ActionMappings=(ActionName="Handbrake",Key=SpaceBar,bShift=False,bCtrl=False,bAlt=False,bCmd=False)
-+ActionMappings=(ActionName="Handbrake",Key=Gamepad_LeftShoulder,bShift=False,bCtrl=False,bAlt=False,bCmd=False)
-+ActionMappings=(ActionName="SwitchCamera",Key=Tab,bShift=False,bCtrl=False,bAlt=False,bCmd=False)
-+ActionMappings=(ActionName="ResetVR",Key=R,bShift=False,bCtrl=True,bAlt=False,bCmd=False)
-+ActionMappings=(ActionName="ResetVR",Key=MotionController_Left_Grip1,bShift=False,bCtrl=False,bAlt=False,bCmd=False)
-+ActionMappings=(ActionName="SwitchCamera",Key=Gamepad_Special_Left,bShift=False,bCtrl=False,bAlt=False,bCmd=False)
-+AxisMappings=(AxisName="MoveForward",Key=W,Scale=1.000000)
-+AxisMappings=(AxisName="MoveForward",Key=S,Scale=-1.000000)
-+AxisMappings=(AxisName="MoveForward",Key=Up,Scale=1.000000)
-+AxisMappings=(AxisName="MoveForward",Key=Down,Scale=-1.000000)
-+AxisMappings=(AxisName="MoveForward",Key=Gamepad_FaceButton_Bottom,Scale=1.000000)
-+AxisMappings=(AxisName="MoveRight",Key=A,Scale=-1.000000)
-+AxisMappings=(AxisName="MoveRight",Key=D,Scale=1.000000)
-+AxisMappings=(AxisName="MoveRight",Key=Left,Scale=-1.000000)
-+AxisMappings=(AxisName="MoveRight",Key=Right,Scale=1.000000)
-+AxisMappings=(AxisName="MoveRight",Key=Gamepad_LeftX,Scale=1.000000)
-+AxisMappings=(AxisName="LookUp",Key=MouseY,Scale=1.000000)
-+AxisMappings=(AxisName="LookRight",Key=MouseX,Scale=1.000000)
-+AxisMappings=(AxisName="MoveForward",Key=Gamepad_RightTriggerAxis,Scale=1.000000)
-+AxisMappings=(AxisName="MoveForward",Key=Gamepad_LeftTriggerAxis,Scale=-1.000000)
-+AxisMappings=(AxisName="LookUp",Key=Gamepad_RightY,Scale=-1.000000)
-+AxisMappings=(AxisName="LookRight",Key=Gamepad_RightX,Scale=1.000000)
-+AxisMappings=(AxisName="MoveForward",Key=Gamepad_LeftY,Scale=1.000000)
-+AxisMappings=(AxisName="MoveForward",Key=GenericUSBController_Axis2,Scale=-2.000000)
-+AxisMappings=(AxisName="MoveForward",Key=GenericUSBController_Axis3,Scale=1.000000)
-+AxisMappings=(AxisName="MoveRight",Key=GenericUSBController_Axis1,Scale=3.000000)
-DefaultTouchInterface=/Game/VehicleAdv/Vehicle/LeftVirtualJoystickOnly.LeftVirtualJoystickOnly
-ConsoleKey=None
+bEnableMouseSmoothing=True
+bEnableFOVScaling=True
+bCaptureMouseOnLaunch=True
+bAlwaysShowTouchInterface=False
+bShowConsoleOnFourFingerTap=True
+bEnableGestureRecognizer=False
+bUseAutocorrect=False
+DefaultViewportMouseCaptureMode=CapturePermanently_IncludingInitialMouseDown
+DefaultViewportMouseLockMode=LockOnCapture
+FOVScale=0.011110
+DoubleClickTime=0.200000
++ActionMappings=(ActionName="Handbrake",bShift=False,bCtrl=False,bAlt=False,bCmd=False,Key=SpaceBar)
++ActionMappings=(ActionName="Handbrake",bShift=False,bCtrl=False,bAlt=False,bCmd=False,Key=Gamepad_LeftShoulder)
++ActionMappings=(ActionName="SwitchCamera",bShift=False,bCtrl=False,bAlt=False,bCmd=False,Key=Tab)
++ActionMappings=(ActionName="ResetVR",bShift=False,bCtrl=False,bAlt=False,bCmd=False,Key=R)
++ActionMappings=(ActionName="SwitchCamera",bShift=False,bCtrl=False,bAlt=False,bCmd=False,Key=Gamepad_Special_Left)
++ActionMappings=(ActionName="ResetVR",bShift=False,bCtrl=False,bAlt=False,bCmd=False,Key=Vive_Left_Menu_Click)
++ActionMappings=(ActionName="ResetVR",bShift=False,bCtrl=False,bAlt=False,bCmd=False,Key=MixedReality_Left_Thumbstick_Click)
++ActionMappings=(ActionName="ResetVR",bShift=False,bCtrl=False,bAlt=False,bCmd=False,Key=OculusTouch_Left_Thumbstick_Click)
++ActionMappings=(ActionName="ResetVR",bShift=False,bCtrl=False,bAlt=False,bCmd=False,Key=ValveIndex_Left_Thumbstick_Click)
++ActionMappings=(ActionName="ResetVR",bShift=False,bCtrl=False,bAlt=False,bCmd=False,Key=MagicLeap_Left_Bumper)
++ActionMappings=(ActionName="Handbrake",bShift=False,bCtrl=False,bAlt=False,bCmd=False,Key=Vive_Right_Grip_Click)
++ActionMappings=(ActionName="Handbrake",bShift=False,bCtrl=False,bAlt=False,bCmd=False,Key=MixedReality_Right_Grip_Click)
++ActionMappings=(ActionName="Handbrake",bShift=False,bCtrl=False,bAlt=False,bCmd=False,Key=OculusTouch_Right_Grip_Click)
++ActionMappings=(ActionName="Handbrake",bShift=False,bCtrl=False,bAlt=False,bCmd=False,Key=ValveIndex_Right_A_Click)
++ActionMappings=(ActionName="Handbrake",bShift=False,bCtrl=False,bAlt=False,bCmd=False,Key=MagicLeap_Right_Bumper)
++AxisMappings=(AxisName="MoveForward",Scale=1.000000,Key=W)
++AxisMappings=(AxisName="MoveForward",Scale=-1.000000,Key=S)
++AxisMappings=(AxisName="MoveForward",Scale=1.000000,Key=Up)
++AxisMappings=(AxisName="MoveForward",Scale=-1.000000,Key=Down)
++AxisMappings=(AxisName="MoveForward",Scale=1.000000,Key=Gamepad_FaceButton_Bottom)
++AxisMappings=(AxisName="MoveForward",Scale=-1.000000,Key=Gamepad_RightY)
++AxisMappings=(AxisName="MoveRight",Scale=-1.000000,Key=A)
++AxisMappings=(AxisName="MoveRight",Scale=1.000000,Key=D)
++AxisMappings=(AxisName="MoveRight",Scale=-1.000000,Key=Left)
++AxisMappings=(AxisName="MoveRight",Scale=1.000000,Key=Right)
++AxisMappings=(AxisName="MoveRight",Scale=1.000000,Key=Gamepad_LeftX)
++AxisMappings=(AxisName="LookUp",Scale=1.000000,Key=MouseY)
++AxisMappings=(AxisName="LookRight",Scale=1.000000,Key=MouseX)
++AxisMappings=(AxisName="MoveForward",Scale=1.000000,Key=Gamepad_RightTriggerAxis)
++AxisMappings=(AxisName="MoveForward",Scale=-1.000000,Key=Gamepad_LeftTriggerAxis)
++AxisMappings=(AxisName="MoveForward",Scale=-2.000000,Key=GenericUSBController_Axis2)
++AxisMappings=(AxisName="MoveForward",Scale=1.000000,Key=GenericUSBController_Axis3)
++AxisMappings=(AxisName="MoveRight",Scale=3.000000,Key=GenericUSBController_Axis1)
++AxisMappings=(AxisName="MoveForward",Scale=1.000000,Key=Vive_Right_Trigger_Axis)
++AxisMappings=(AxisName="MoveForward",Scale=-1.000000,Key=Vive_Left_Trigger_Axis)
++AxisMappings=(AxisName="MoveForward",Scale=1.000000,Key=MixedReality_Right_Trigger_Axis)
++AxisMappings=(AxisName="MoveForward",Scale=-1.000000,Key=MixedReality_Left_Trigger_Axis)
++AxisMappings=(AxisName="MoveForward",Scale=1.000000,Key=OculusTouch_Right_Trigger_Axis)
++AxisMappings=(AxisName="MoveForward",Scale=-1.000000,Key=OculusTouch_Left_Trigger_Axis)
++AxisMappings=(AxisName="MoveForward",Scale=1.000000,Key=ValveIndex_Right_Trigger_Axis)
++AxisMappings=(AxisName="MoveForward",Scale=-1.000000,Key=ValveIndex_Left_Trigger_Axis)
++AxisMappings=(AxisName="MoveForward",Scale=1.000000,Key=MagicLeap_Right_Trigger_Axis)
++AxisMappings=(AxisName="MoveForward",Scale=-1.000000,Key=MagicLeap_Left_Trigger_Axis)
++AxisMappings=(AxisName="MoveRight",Scale=1.000000,Key=Vive_Left_Trackpad_X)
++AxisMappings=(AxisName="MoveRight",Scale=1.000000,Key=Vive_Right_Trackpad_X)
++AxisMappings=(AxisName="MoveRight",Scale=1.000000,Key=MixedReality_Left_Thumbstick_X)
++AxisMappings=(AxisName="MoveRight",Scale=1.000000,Key=MixedReality_Right_Thumbstick_X)
++AxisMappings=(AxisName="MoveRight",Scale=1.000000,Key=OculusTouch_Left_Thumbstick_X)
++AxisMappings=(AxisName="MoveRight",Scale=1.000000,Key=OculusTouch_Right_Thumbstick_X)
++AxisMappings=(AxisName="MoveRight",Scale=1.000000,Key=ValveIndex_Left_Thumbstick_X)
++AxisMappings=(AxisName="MoveRight",Scale=1.000000,Key=ValveIndex_Right_Thumbstick_X)
++AxisMappings=(AxisName="MoveRight",Scale=1.000000,Key=MagicLeap_Left_Trackpad_X)
++AxisMappings=(AxisName="MoveRight",Scale=1.000000,Key=MagicLeap_Right_Trackpad_X)
+DefaultTouchInterface=/Game/Vehicle/Sedan/SedanVirtualJoysticks.SedanVirtualJoysticks
 +ConsoleKeys=Tilde
 
-[/Script/RawInput.RawInputSettings]
-+DeviceConfigurations=(VendorID="0x046D",ProductID="0xC262",AxisProperties=((bEnabled=False),(Key=GenericUSBController_Axis1,Offset=-0.500000),(bEnabled=False),(Key=GenericUSBController_Axis2,bInverted=True,Offset=1.000000),(Key=GenericUSBController_Axis3,bInverted=True,Offset=1.000000)),ButtonProperties=((Key=GenericUSBController_Button1),(Key=GenericUSBController_Button2),(Key=GenericUSBController_Button3),(Key=GenericUSBController_Button4),(Key=GenericUSBController_Button5),(Key=GenericUSBController_Button6),(Key=GenericUSBController_Button7),(Key=GenericUSBController_Button8),(Key=GenericUSBController_Button9),(Key=GenericUSBController_Button10),(Key=GenericUSBController_Button11),(Key=GenericUSBController_Button12)))
-
 
diff --git a/Templates/TP_VehicleAdvBP/Config/DefaultInput.ini b/Templates/TP_VehicleAdvBP/Config/DefaultInput.ini
--- a/Templates/TP_VehicleAdvBP/Config/DefaultInput.ini
+++ b/Templates/TP_VehicleAdvBP/Config/DefaultInput.ini
@@ -1,11 +1,7 @@
+[/Script/RawInput.RawInputSettings]
++DeviceConfigurations=(VendorID="0x046D",ProductID="0xC262",AxisProperties=((bEnabled=False),(Key=GenericUSBController_Axis1,Offset=-0.500000),(bEnabled=False),(Key=GenericUSBController_Axis2,bInverted=True,Offset=1.000000),(Key=GenericUSBController_Axis3,bInverted=True,Offset=1.000000)),ButtonProperties=((Key=GenericUSBController_Button1),(Key=GenericUSBController_Button2),(Key=GenericUSBController_Button3),(Key=GenericUSBController_Button4),(Key=GenericUSBController_Button5),(Key=GenericUSBController_Button6),(Key=GenericUSBController_Button7),(Key=GenericUSBController_Button8),(Key=GenericUSBController_Button9),(Key=GenericUSBController_Button10),(Key=GenericUSBController_Button11),(Key=GenericUSBController_Button12)))
 
 [/Script/Engine.InputSettings]
--AxisConfig=(AxisKeyName="Gamepad_LeftX",AxisProperties=(DeadZone=0.25,Exponent=1.f,Sensitivity=1.f))
--AxisConfig=(AxisKeyName="Gamepad_LeftY",AxisProperties=(DeadZone=0.25,Exponent=1.f,Sensitivity=1.f))
--AxisConfig=(AxisKeyName="Gamepad_RightX",AxisProperties=(DeadZone=0.25,Exponent=1.f,Sensitivity=1.f))
--AxisConfig=(AxisKeyName="Gamepad_RightY",AxisProperties=(DeadZone=0.25,Exponent=1.f,Sensitivity=1.f))
--AxisConfig=(AxisKeyName="MouseX",AxisProperties=(DeadZone=0.f,Exponent=1.f,Sensitivity=0.07f))
--AxisConfig=(AxisKeyName="MouseY",AxisProperties=(DeadZone=0.f,Exponent=1.f,Sensitivity=0.07f))
 +AxisConfig=(AxisKeyName="Gamepad_LeftX",AxisProperties=(DeadZone=0.250000,Sensitivity=1.000000,Exponent=1.000000,bInvert=False))
 +AxisConfig=(AxisKeyName="Gamepad_LeftY",AxisProperties=(DeadZone=0.250000,Sensitivity=1.000000,Exponent=1.000000,bInvert=False))
 +AxisConfig=(AxisKeyName="Gamepad_RightX",AxisProperties=(DeadZone=0.250000,Sensitivity=1.000000,Exponent=1.000000,bInvert=False))
@@ -27,57 +23,154 @@
 +AxisConfig=(AxisKeyName="MotionController_Right_Grip2Axis",AxisProperties=(DeadZone=0.000000,Sensitivity=1.000000,Exponent=1.000000,bInvert=False))
 +AxisConfig=(AxisKeyName="Gamepad_Special_Left_X",AxisProperties=(DeadZone=0.000000,Sensitivity=1.000000,Exponent=1.000000,bInvert=False))
 +AxisConfig=(AxisKeyName="Gamepad_Special_Left_Y",AxisProperties=(DeadZone=0.000000,Sensitivity=1.000000,Exponent=1.000000,bInvert=False))
++AxisConfig=(AxisKeyName="Daydream_Left_Trackpad_X",AxisProperties=(DeadZone=0.000000,Sensitivity=1.000000,Exponent=1.000000,bInvert=False))
++AxisConfig=(AxisKeyName="Daydream_Left_Trackpad_Y",AxisProperties=(DeadZone=0.000000,Sensitivity=1.000000,Exponent=1.000000,bInvert=False))
++AxisConfig=(AxisKeyName="Daydream_Right_Trackpad_X",AxisProperties=(DeadZone=0.000000,Sensitivity=1.000000,Exponent=1.000000,bInvert=False))
++AxisConfig=(AxisKeyName="Daydream_Right_Trackpad_Y",AxisProperties=(DeadZone=0.000000,Sensitivity=1.000000,Exponent=1.000000,bInvert=False))
++AxisConfig=(AxisKeyName="Vive_Left_Trigger_Axis",AxisProperties=(DeadZone=0.000000,Sensitivity=1.000000,Exponent=1.000000,bInvert=False))
++AxisConfig=(AxisKeyName="Vive_Left_Trackpad_X",AxisProperties=(DeadZone=0.000000,Sensitivity=1.000000,Exponent=1.000000,bInvert=False))
++AxisConfig=(AxisKeyName="Vive_Left_Trackpad_Y",AxisProperties=(DeadZone=0.000000,Sensitivity=1.000000,Exponent=1.000000,bInvert=False))
++AxisConfig=(AxisKeyName="Vive_Right_Trigger_Axis",AxisProperties=(DeadZone=0.000000,Sensitivity=1.000000,Exponent=1.000000,bInvert=False))
++AxisConfig=(AxisKeyName="Vive_Right_Trackpad_X",AxisProperties=(DeadZone=0.000000,Sensitivity=1.000000,Exponent=1.000000,bInvert=False))
++AxisConfig=(AxisKeyName="Vive_Right_Trackpad_Y",AxisProperties=(DeadZone=0.000000,Sensitivity=1.000000,Exponent=1.000000,bInvert=False))
++AxisConfig=(AxisKeyName="MixedReality_Left_Trigger_Axis",AxisProperties=(DeadZone=0.000000,Sensitivity=1.000000,Exponent=1.000000,bInvert=False))
++AxisConfig=(AxisKeyName="MixedReality_Left_Thumbstick_X",AxisProperties=(DeadZone=0.000000,Sensitivity=1.000000,Exponent=1.000000,bInvert=False))
++AxisConfig=(AxisKeyName="MixedReality_Left_Thumbstick_Y",AxisProperties=(DeadZone=0.000000,Sensitivity=1.000000,Exponent=1.000000,bInvert=False))
++AxisConfig=(AxisKeyName="MixedReality_Left_Trackpad_X",AxisProperties=(DeadZone=0.000000,Sensitivity=1.000000,Exponent=1.000000,bInvert=False))
++AxisConfig=(AxisKeyName="MixedReality_Left_Trackpad_Y",AxisProperties=(DeadZone=0.000000,Sensitivity=1.000000,Exponent=1.000000,bInvert=False))
++AxisConfig=(AxisKeyName="MixedReality_Right_Trigger_Axis",AxisProperties=(DeadZone=0.000000,Sensitivity=1.000000,Exponent=1.000000,bInvert=False))
++AxisConfig=(AxisKeyName="MixedReality_Right_Thumbstick_X",AxisProperties=(DeadZone=0.000000,Sensitivity=1.000000,Exponent=1.000000,bInvert=False))
++AxisConfig=(AxisKeyName="MixedReality_Right_Thumbstick_Y",AxisProperties=(DeadZone=0.000000,Sensitivity=1.000000,Exponent=1.000000,bInvert=False))
++AxisConfig=(AxisKeyName="MixedReality_Right_Trackpad_X",AxisProperties=(DeadZone=0.000000,Sensitivity=1.000000,Exponent=1.000000,bInvert=False))
++AxisConfig=(AxisKeyName="MixedReality_Right_Trackpad_Y",AxisProperties=(DeadZone=0.000000,Sensitivity=1.000000,Exponent=1.000000,bInvert=False))
++AxisConfig=(AxisKeyName="OculusGo_Left_Trackpad_X",AxisProperties=(DeadZone=0.000000,Sensitivity=1.000000,Exponent=1.000000,bInvert=False))
++AxisConfig=(AxisKeyName="OculusGo_Left_Trackpad_Y",AxisProperties=(DeadZone=0.000000,Sensitivity=1.000000,Exponent=1.000000,bInvert=False))
++AxisConfig=(AxisKeyName="OculusGo_Right_Trackpad_X",AxisProperties=(DeadZone=0.000000,Sensitivity=1.000000,Exponent=1.000000,bInvert=False))
++AxisConfig=(AxisKeyName="OculusGo_Right_Trackpad_Y",AxisProperties=(DeadZone=0.000000,Sensitivity=1.000000,Exponent=1.000000,bInvert=False))
++AxisConfig=(AxisKeyName="OculusTouch_Left_Grip_Axis",AxisProperties=(DeadZone=0.000000,Sensitivity=1.000000,Exponent=1.000000,bInvert=False))
++AxisConfig=(AxisKeyName="OculusTouch_Left_Trigger_Axis",AxisProperties=(DeadZone=0.000000,Sensitivity=1.000000,Exponent=1.000000,bInvert=False))
++AxisConfig=(AxisKeyName="OculusTouch_Left_Thumbstick_X",AxisProperties=(DeadZone=0.000000,Sensitivity=1.000000,Exponent=1.000000,bInvert=False))
++AxisConfig=(AxisKeyName="OculusTouch_Left_Thumbstick_Y",AxisProperties=(DeadZone=0.000000,Sensitivity=1.000000,Exponent=1.000000,bInvert=False))
++AxisConfig=(AxisKeyName="OculusTouch_Right_Grip_Axis",AxisProperties=(DeadZone=0.000000,Sensitivity=1.000000,Exponent=1.000000,bInvert=False))
++AxisConfig=(AxisKeyName="OculusTouch_Right_Trigger_Axis",AxisProperties=(DeadZone=0.000000,Sensitivity=1.000000,Exponent=1.000000,bInvert=False))
++AxisConfig=(AxisKeyName="OculusTouch_Right_Thumbstick_X",AxisProperties=(DeadZone=0.000000,Sensitivity=1.000000,Exponent=1.000000,bInvert=False))
++AxisConfig=(AxisKeyName="OculusTouch_Right_Thumbstick_Y",AxisProperties=(DeadZone=0.000000,Sensitivity=1.000000,Exponent=1.000000,bInvert=False))
++AxisConfig=(AxisKeyName="ValveIndex_Left_Grip_Axis",AxisProperties=(DeadZone=0.000000,Sensitivity=1.000000,Exponent=1.000000,bInvert=False))
++AxisConfig=(AxisKeyName="ValveIndex_Left_Grip_Force",AxisProperties=(DeadZone=0.000000,Sensitivity=1.000000,Exponent=1.000000,bInvert=False))
++AxisConfig=(AxisKeyName="ValveIndex_Left_Trigger_Axis",AxisProperties=(DeadZone=0.000000,Sensitivity=1.000000,Exponent=1.000000,bInvert=False))
++AxisConfig=(AxisKeyName="ValveIndex_Left_Thumbstick_X",AxisProperties=(DeadZone=0.000000,Sensitivity=1.000000,Exponent=1.000000,bInvert=False))
++AxisConfig=(AxisKeyName="ValveIndex_Left_Thumbstick_Y",AxisProperties=(DeadZone=0.000000,Sensitivity=1.000000,Exponent=1.000000,bInvert=False))
++AxisConfig=(AxisKeyName="ValveIndex_Left_Trackpad_X",AxisProperties=(DeadZone=0.000000,Sensitivity=1.000000,Exponent=1.000000,bInvert=False))
++AxisConfig=(AxisKeyName="ValveIndex_Left_Trackpad_Y",AxisProperties=(DeadZone=0.000000,Sensitivity=1.000000,Exponent=1.000000,bInvert=False))
++AxisConfig=(AxisKeyName="ValveIndex_Left_Trackpad_Force",AxisProperties=(DeadZone=0.000000,Sensitivity=1.000000,Exponent=1.000000,bInvert=False))
++AxisConfig=(AxisKeyName="ValveIndex_Left_Trackpad_Touch",AxisProperties=(DeadZone=0.000000,Sensitivity=1.000000,Exponent=1.000000,bInvert=False))
++AxisConfig=(AxisKeyName="ValveIndex_Right_Grip_Axis",AxisProperties=(DeadZone=0.000000,Sensitivity=1.000000,Exponent=1.000000,bInvert=False))
++AxisConfig=(AxisKeyName="ValveIndex_Right_Grip_Force",AxisProperties=(DeadZone=0.000000,Sensitivity=1.000000,Exponent=1.000000,bInvert=False))
++AxisConfig=(AxisKeyName="ValveIndex_Right_Trigger_Axis",AxisProperties=(DeadZone=0.000000,Sensitivity=1.000000,Exponent=1.000000,bInvert=False))
++AxisConfig=(AxisKeyName="ValveIndex_Right_Thumbstick_X",AxisProperties=(DeadZone=0.000000,Sensitivity=1.000000,Exponent=1.000000,bInvert=False))
++AxisConfig=(AxisKeyName="ValveIndex_Right_Thumbstick_Y",AxisProperties=(DeadZone=0.000000,Sensitivity=1.000000,Exponent=1.000000,bInvert=False))
++AxisConfig=(AxisKeyName="ValveIndex_Right_Trackpad_X",AxisProperties=(DeadZone=0.000000,Sensitivity=1.000000,Exponent=1.000000,bInvert=False))
++AxisConfig=(AxisKeyName="ValveIndex_Right_Trackpad_Y",AxisProperties=(DeadZone=0.000000,Sensitivity=1.000000,Exponent=1.000000,bInvert=False))
++AxisConfig=(AxisKeyName="ValveIndex_Right_Trackpad_Force",AxisProperties=(DeadZone=0.000000,Sensitivity=1.000000,Exponent=1.000000,bInvert=False))
++AxisConfig=(AxisKeyName="MotionController_Left_Thumbstick_Z",AxisProperties=(DeadZone=0.000000,Sensitivity=1.000000,Exponent=1.000000,bInvert=False))
++AxisConfig=(AxisKeyName="MagicLeap_Left_Trigger_Axis",AxisProperties=(DeadZone=0.000000,Sensitivity=1.000000,Exponent=1.000000,bInvert=False))
++AxisConfig=(AxisKeyName="MagicLeap_Left_Trackpad_X",AxisProperties=(DeadZone=0.000000,Sensitivity=1.000000,Exponent=1.000000,bInvert=False))
++AxisConfig=(AxisKeyName="MagicLeap_Left_Trackpad_Y",AxisProperties=(DeadZone=0.000000,Sensitivity=1.000000,Exponent=1.000000,bInvert=False))
++AxisConfig=(AxisKeyName="MagicLeap_Left_Trackpad_Force",AxisProperties=(DeadZone=0.000000,Sensitivity=1.000000,Exponent=1.000000,bInvert=False))
++AxisConfig=(AxisKeyName="MagicLeap_Left_Touch1_X",AxisProperties=(DeadZone=0.000000,Sensitivity=1.000000,Exponent=1.000000,bInvert=False))
++AxisConfig=(AxisKeyName="MagicLeap_Left_Touch1_Y",AxisProperties=(DeadZone=0.000000,Sensitivity=1.000000,Exponent=1.000000,bInvert=False))
++AxisConfig=(AxisKeyName="MagicLeap_Left_Touch1_Force",AxisProperties=(DeadZone=0.000000,Sensitivity=1.000000,Exponent=1.000000,bInvert=False))
++AxisConfig=(AxisKeyName="MotionController_Right_Thumbstick_Z",AxisProperties=(DeadZone=0.000000,Sensitivity=1.000000,Exponent=1.000000,bInvert=False))
++AxisConfig=(AxisKeyName="MagicLeap_Right_Trigger_Axis",AxisProperties=(DeadZone=0.000000,Sensitivity=1.000000,Exponent=1.000000,bInvert=False))
++AxisConfig=(AxisKeyName="MagicLeap_Right_Trackpad_X",AxisProperties=(DeadZone=0.000000,Sensitivity=1.000000,Exponent=1.000000,bInvert=False))
++AxisConfig=(AxisKeyName="MagicLeap_Right_Trackpad_Y",AxisProperties=(DeadZone=0.000000,Sensitivity=1.000000,Exponent=1.000000,bInvert=False))
++AxisConfig=(AxisKeyName="MagicLeap_Right_Trackpad_Force",AxisProperties=(DeadZone=0.000000,Sensitivity=1.000000,Exponent=1.000000,bInvert=False))
++AxisConfig=(AxisKeyName="MagicLeap_Right_Touch1_X",AxisProperties=(DeadZone=0.000000,Sensitivity=1.000000,Exponent=1.000000,bInvert=False))
++AxisConfig=(AxisKeyName="MagicLeap_Right_Touch1_Y",AxisProperties=(DeadZone=0.000000,Sensitivity=1.000000,Exponent=1.000000,bInvert=False))
++AxisConfig=(AxisKeyName="MagicLeap_Right_Touch1_Force",AxisProperties=(DeadZone=0.000000,Sensitivity=1.000000,Exponent=1.000000,bInvert=False))
++AxisConfig=(AxisKeyName="OculusTouch_Left_Thumbstick",AxisProperties=(DeadZone=0.000000,Sensitivity=1.000000,Exponent=1.000000,bInvert=False))
++AxisConfig=(AxisKeyName="OculusTouch_Left_FaceButton1",AxisProperties=(DeadZone=0.000000,Sensitivity=1.000000,Exponent=1.000000,bInvert=False))
++AxisConfig=(AxisKeyName="OculusTouch_Left_Trigger",AxisProperties=(DeadZone=0.000000,Sensitivity=1.000000,Exponent=1.000000,bInvert=False))
++AxisConfig=(AxisKeyName="OculusTouch_Left_FaceButton2",AxisProperties=(DeadZone=0.000000,Sensitivity=1.000000,Exponent=1.000000,bInvert=False))
++AxisConfig=(AxisKeyName="OculusTouch_Left_IndexPointing",AxisProperties=(DeadZone=0.000000,Sensitivity=1.000000,Exponent=1.000000,bInvert=False))
++AxisConfig=(AxisKeyName="OculusTouch_Left_ThumbUp",AxisProperties=(DeadZone=0.000000,Sensitivity=1.000000,Exponent=1.000000,bInvert=False))
++AxisConfig=(AxisKeyName="OculusTouch_Right_Thumbstick",AxisProperties=(DeadZone=0.000000,Sensitivity=1.000000,Exponent=1.000000,bInvert=False))
++AxisConfig=(AxisKeyName="OculusTouch_Right_FaceButton1",AxisProperties=(DeadZone=0.000000,Sensitivity=1.000000,Exponent=1.000000,bInvert=False))
++AxisConfig=(AxisKeyName="OculusTouch_Right_Trigger",AxisProperties=(DeadZone=0.000000,Sensitivity=1.000000,Exponent=1.000000,bInvert=False))
++AxisConfig=(AxisKeyName="OculusTouch_Right_FaceButton2",AxisProperties=(DeadZone=0.000000,Sensitivity=1.000000,Exponent=1.000000,bInvert=False))
++AxisConfig=(AxisKeyName="OculusTouch_Right_IndexPointing",AxisProperties=(DeadZone=0.000000,Sensitivity=1.000000,Exponent=1.000000,bInvert=False))
++AxisConfig=(AxisKeyName="OculusTouch_Right_ThumbUp",AxisProperties=(DeadZone=0.000000,Sensitivity=1.000000,Exponent=1.000000,bInvert=False))
++AxisConfig=(AxisKeyName="OculusTouchpad_Touchpad_X",AxisProperties=(DeadZone=0.000000,Sensitivity=1.000000,Exponent=1.000000,bInvert=False))
++AxisConfig=(AxisKeyName="OculusTouchpad_Touchpad_Y",AxisProperties=(DeadZone=0.000000,Sensitivity=1.000000,Exponent=1.000000,bInvert=False))
 bAltEnterTogglesFullscreen=True
 bF11TogglesFullscreen=True
 bUseMouseForTouch=False
 bEnableMouseSmoothing=True
 bEnableFOVScaling=True
-FOVScale=0.011110
-DoubleClickTime=0.200000
 bCaptureMouseOnLaunch=True
+bAlwaysShowTouchInterface=False
+bShowConsoleOnFourFingerTap=True
+bEnableGestureRecognizer=False
+bUseAutocorrect=False
 DefaultViewportMouseCaptureMode=CapturePermanently_IncludingInitialMouseDown
-bDefaultViewportMouseLock=True
 DefaultViewportMouseLockMode=LockOnCapture
-+ActionMappings=(ActionName="Handbrake",Key=SpaceBar,bShift=False,bCtrl=False,bAlt=False,bCmd=False)
-+ActionMappings=(ActionName="SwitchCamera",Key=Tab,bShift=False,bCtrl=False,bAlt=False,bCmd=False)
-+ActionMappings=(ActionName="ResetVR",Key=R,bShift=False,bCtrl=True,bAlt=False,bCmd=False)
+FOVScale=0.011110
+DoubleClickTime=0.200000
++ActionMappings=(ActionName="Handbrake",bShift=False,bCtrl=False,bAlt=False,bCmd=False,Key=SpaceBar)
++ActionMappings=(ActionName="Handbrake",bShift=False,bCtrl=False,bAlt=False,bCmd=False,Key=Gamepad_LeftShoulder)
++ActionMappings=(ActionName="SwitchCamera",bShift=False,bCtrl=False,bAlt=False,bCmd=False,Key=Tab)
++ActionMappings=(ActionName="ResetVR",bShift=False,bCtrl=False,bAlt=False,bCmd=False,Key=R)
 +ActionMappings=(ActionName="ResetCar",Key=BackSpace,bShift=False,bCtrl=False,bAlt=False,bCmd=False)
-+ActionMappings=(ActionName="Handbrake",Key=Gamepad_LeftShoulder,bShift=False,bCtrl=False,bAlt=False,bCmd=False)
-+ActionMappings=(ActionName="SwitchCamera",Key=Gamepad_Special_Left,bShift=False,bCtrl=False,bAlt=False,bCmd=False)
-+ActionMappings=(ActionName="ResetVR",Key=MotionController_Left_Grip1,bShift=False,bCtrl=False,bAlt=False,bCmd=False)
-+ActionMappings=(ActionName="Handbrake",Key=MotionController_Right_Grip1,bShift=False,bCtrl=False,bAlt=False,bCmd=False)
-+AxisMappings=(AxisName="MoveForward",Key=W,Scale=1.000000)
-+AxisMappings=(AxisName="MoveRight",Key=D,Scale=1.000000)
-+AxisMappings=(AxisName="MoveForward",Key=S,Scale=-1.000000)
-+AxisMappings=(AxisName="MoveRight",Key=A,Scale=-1.000000)
-+AxisMappings=(AxisName="MoveForward",Key=Up,Scale=1.000000)
-+AxisMappings=(AxisName="MoveForward",Key=Down,Scale=-1.000000)
-+AxisMappings=(AxisName="MoveRight",Key=Right,Scale=1.000000)
-+AxisMappings=(AxisName="MoveRight",Key=Left,Scale=-1.000000)
-+AxisMappings=(AxisName="LookUp",Key=MouseY,Scale=1.000000)
-+AxisMappings=(AxisName="LookRight",Key=MouseX,Scale=1.000000)
-+AxisMappings=(AxisName="LookUp",Key=Gamepad_RightY,Scale=-1.000000)
-+AxisMappings=(AxisName="LookRight",Key=Gamepad_RightX,Scale=1.000000)
-+AxisMappings=(AxisName="MoveForward",Key=Gamepad_LeftY,Scale=1.000000)
-+AxisMappings=(AxisName="MoveRight",Key=Gamepad_LeftX,Scale=1.000000)
-+AxisMappings=(AxisName="MoveForward",Key=Gamepad_RightTrigger,Scale=1.000000)
-+AxisMappings=(AxisName="MoveForward",Key=Gamepad_FaceButton_Bottom,Scale=1.000000)
-+AxisMappings=(AxisName="MoveForward",Key=Gamepad_LeftTrigger,Scale=-1.000000)
-+AxisMappings=(AxisName="MoveForward",Key=MotionController_Right_Trigger,Scale=1.000000)
-+AxisMappings=(AxisName="MoveForward",Key=MotionController_Left_Trigger,Scale=-1.000000)
-+AxisMappings=(AxisName="MoveRight",Key=MotionController_Left_Thumbstick_X,Scale=1.000000)
-+AxisMappings=(AxisName="MoveRight",Key=MotionController_Right_Thumbstick_X,Scale=1.000000)
-+AxisMappings=(AxisName="MoveForward",Key=GenericUSBController_Axis2,Scale=-2.000000)
-+AxisMappings=(AxisName="MoveForward",Key=GenericUSBController_Axis3,Scale=1.000000)
-+AxisMappings=(AxisName="MoveRight",Key=GenericUSBController_Axis1,Scale=3.000000)
-bAlwaysShowTouchInterface=False
-bShowConsoleOnFourFingerTap=True
-DefaultTouchInterface=/Engine/MobileResources/HUD/LeftVirtualJoystickOnly.LeftVirtualJoystickOnly
-ConsoleKey=Tilde
--ConsoleKeys=Tilde
++ActionMappings=(ActionName="SwitchCamera",bShift=False,bCtrl=False,bAlt=False,bCmd=False,Key=Gamepad_Special_Left)
++ActionMappings=(ActionName="ResetVR",bShift=False,bCtrl=False,bAlt=False,bCmd=False,Key=Vive_Left_Menu_Click)
++ActionMappings=(ActionName="ResetVR",bShift=False,bCtrl=False,bAlt=False,bCmd=False,Key=MixedReality_Left_Thumbstick_Click)
++ActionMappings=(ActionName="ResetVR",bShift=False,bCtrl=False,bAlt=False,bCmd=False,Key=OculusTouch_Left_Thumbstick_Click)
++ActionMappings=(ActionName="ResetVR",bShift=False,bCtrl=False,bAlt=False,bCmd=False,Key=ValveIndex_Left_Thumbstick_Click)
++ActionMappings=(ActionName="ResetVR",bShift=False,bCtrl=False,bAlt=False,bCmd=False,Key=MagicLeap_Left_Bumper)
++ActionMappings=(ActionName="Handbrake",bShift=False,bCtrl=False,bAlt=False,bCmd=False,Key=Vive_Right_Grip_Click)
++ActionMappings=(ActionName="Handbrake",bShift=False,bCtrl=False,bAlt=False,bCmd=False,Key=MixedReality_Right_Grip_Click)
++ActionMappings=(ActionName="Handbrake",bShift=False,bCtrl=False,bAlt=False,bCmd=False,Key=OculusTouch_Right_Grip_Click)
++ActionMappings=(ActionName="Handbrake",bShift=False,bCtrl=False,bAlt=False,bCmd=False,Key=ValveIndex_Right_A_Click)
++ActionMappings=(ActionName="Handbrake",bShift=False,bCtrl=False,bAlt=False,bCmd=False,Key=MagicLeap_Right_Bumper)
++AxisMappings=(AxisName="MoveForward",Scale=1.000000,Key=W)
++AxisMappings=(AxisName="MoveForward",Scale=-1.000000,Key=S)
++AxisMappings=(AxisName="MoveForward",Scale=1.000000,Key=Up)
++AxisMappings=(AxisName="MoveForward",Scale=-1.000000,Key=Down)
++AxisMappings=(AxisName="MoveForward",Scale=1.000000,Key=Gamepad_FaceButton_Bottom)
++AxisMappings=(AxisName="MoveForward",Scale=-1.000000,Key=Gamepad_RightY)
++AxisMappings=(AxisName="MoveRight",Scale=-1.000000,Key=A)
++AxisMappings=(AxisName="MoveRight",Scale=1.000000,Key=D)
++AxisMappings=(AxisName="MoveRight",Scale=-1.000000,Key=Left)
++AxisMappings=(AxisName="MoveRight",Scale=1.000000,Key=Right)
++AxisMappings=(AxisName="MoveRight",Scale=1.000000,Key=Gamepad_LeftX)
++AxisMappings=(AxisName="LookUp",Scale=1.000000,Key=MouseY)
++AxisMappings=(AxisName="LookRight",Scale=1.000000,Key=MouseX)
++AxisMappings=(AxisName="MoveForward",Scale=1.000000,Key=Gamepad_RightTriggerAxis)
++AxisMappings=(AxisName="MoveForward",Scale=-1.000000,Key=Gamepad_LeftTriggerAxis)
++AxisMappings=(AxisName="MoveForward",Scale=-2.000000,Key=GenericUSBController_Axis2)
++AxisMappings=(AxisName="MoveForward",Scale=1.000000,Key=GenericUSBController_Axis3)
++AxisMappings=(AxisName="MoveRight",Scale=3.000000,Key=GenericUSBController_Axis1)
++AxisMappings=(AxisName="MoveForward",Scale=1.000000,Key=Vive_Right_Trigger_Axis)
++AxisMappings=(AxisName="MoveForward",Scale=-1.000000,Key=Vive_Left_Trigger_Axis)
++AxisMappings=(AxisName="MoveForward",Scale=1.000000,Key=MixedReality_Right_Trigger_Axis)
++AxisMappings=(AxisName="MoveForward",Scale=-1.000000,Key=MixedReality_Left_Trigger_Axis)
++AxisMappings=(AxisName="MoveForward",Scale=1.000000,Key=OculusTouch_Right_Trigger_Axis)
++AxisMappings=(AxisName="MoveForward",Scale=-1.000000,Key=OculusTouch_Left_Trigger_Axis)
++AxisMappings=(AxisName="MoveForward",Scale=1.000000,Key=ValveIndex_Right_Trigger_Axis)
++AxisMappings=(AxisName="MoveForward",Scale=-1.000000,Key=ValveIndex_Left_Trigger_Axis)
++AxisMappings=(AxisName="MoveForward",Scale=1.000000,Key=MagicLeap_Right_Trigger_Axis)
++AxisMappings=(AxisName="MoveForward",Scale=-1.000000,Key=MagicLeap_Left_Trigger_Axis)
++AxisMappings=(AxisName="MoveRight",Scale=1.000000,Key=Vive_Left_Trackpad_X)
++AxisMappings=(AxisName="MoveRight",Scale=1.000000,Key=Vive_Right_Trackpad_X)
++AxisMappings=(AxisName="MoveRight",Scale=1.000000,Key=MixedReality_Left_Thumbstick_X)
++AxisMappings=(AxisName="MoveRight",Scale=1.000000,Key=MixedReality_Right_Thumbstick_X)
++AxisMappings=(AxisName="MoveRight",Scale=1.000000,Key=OculusTouch_Left_Thumbstick_X)
++AxisMappings=(AxisName="MoveRight",Scale=1.000000,Key=OculusTouch_Right_Thumbstick_X)
++AxisMappings=(AxisName="MoveRight",Scale=1.000000,Key=ValveIndex_Left_Thumbstick_X)
++AxisMappings=(AxisName="MoveRight",Scale=1.000000,Key=ValveIndex_Right_Thumbstick_X)
++AxisMappings=(AxisName="MoveRight",Scale=1.000000,Key=MagicLeap_Left_Trackpad_X)
++AxisMappings=(AxisName="MoveRight",Scale=1.000000,Key=MagicLeap_Right_Trackpad_X)
+DefaultTouchInterface=/Game/Vehicle/Sedan/SedanVirtualJoysticks.SedanVirtualJoysticks
 +ConsoleKeys=Tilde
 
-[/Script/RawInput.RawInputSettings]
-+DeviceConfigurations=(VendorID="0x046D",ProductID="0xC262",AxisProperties=((bEnabled=False),(Key=GenericUSBController_Axis1,Offset=-0.500000),(bEnabled=False),(Key=GenericUSBController_Axis2,bInverted=True,Offset=1.000000),(Key=GenericUSBController_Axis3,bInverted=True,Offset=1.000000)),ButtonProperties=((Key=GenericUSBController_Button1),(Key=GenericUSBController_Button2),(Key=GenericUSBController_Button3),(Key=GenericUSBController_Button4),(Key=GenericUSBController_Button5),(Key=GenericUSBController_Button6),(Key=GenericUSBController_Button7),(Key=GenericUSBController_Button8),(Key=GenericUSBController_Button9),(Key=GenericUSBController_Button10),(Key=GenericUSBController_Button11),(Key=GenericUSBController_Button12)))
-
 
diff --git a/Templates/TP_VehicleBP/Config/DefaultInput.ini b/Templates/TP_VehicleBP/Config/DefaultInput.ini
--- a/Templates/TP_VehicleBP/Config/DefaultInput.ini
+++ b/Templates/TP_VehicleBP/Config/DefaultInput.ini
@@ -1,11 +1,7 @@
+[/Script/RawInput.RawInputSettings]
++DeviceConfigurations=(VendorID="0x046D",ProductID="0xC262",AxisProperties=((bEnabled=False),(Key=GenericUSBController_Axis1,Offset=-0.500000),(bEnabled=False),(Key=GenericUSBController_Axis2,bInverted=True,Offset=1.000000),(Key=GenericUSBController_Axis3,bInverted=True,Offset=1.000000)),ButtonProperties=((Key=GenericUSBController_Button1),(Key=GenericUSBController_Button2),(Key=GenericUSBController_Button3),(Key=GenericUSBController_Button4),(Key=GenericUSBController_Button5),(Key=GenericUSBController_Button6),(Key=GenericUSBController_Button7),(Key=GenericUSBController_Button8),(Key=GenericUSBController_Button9),(Key=GenericUSBController_Button10),(Key=GenericUSBController_Button11),(Key=GenericUSBController_Button12)))
 
 [/Script/Engine.InputSettings]
--AxisConfig=(AxisKeyName="Gamepad_LeftX",AxisProperties=(DeadZone=0.25,Exponent=1.f,Sensitivity=1.f))
--AxisConfig=(AxisKeyName="Gamepad_LeftY",AxisProperties=(DeadZone=0.25,Exponent=1.f,Sensitivity=1.f))
--AxisConfig=(AxisKeyName="Gamepad_RightX",AxisProperties=(DeadZone=0.25,Exponent=1.f,Sensitivity=1.f))
--AxisConfig=(AxisKeyName="Gamepad_RightY",AxisProperties=(DeadZone=0.25,Exponent=1.f,Sensitivity=1.f))
--AxisConfig=(AxisKeyName="MouseX",AxisProperties=(DeadZone=0.f,Exponent=1.f,Sensitivity=0.07f))
--AxisConfig=(AxisKeyName="MouseY",AxisProperties=(DeadZone=0.f,Exponent=1.f,Sensitivity=0.07f))
 +AxisConfig=(AxisKeyName="Gamepad_LeftX",AxisProperties=(DeadZone=0.250000,Sensitivity=1.000000,Exponent=1.000000,bInvert=False))
 +AxisConfig=(AxisKeyName="Gamepad_LeftY",AxisProperties=(DeadZone=0.250000,Sensitivity=1.000000,Exponent=1.000000,bInvert=False))
 +AxisConfig=(AxisKeyName="Gamepad_RightX",AxisProperties=(DeadZone=0.250000,Sensitivity=1.000000,Exponent=1.000000,bInvert=False))
@@ -27,52 +23,153 @@
 +AxisConfig=(AxisKeyName="MotionController_Right_Grip2Axis",AxisProperties=(DeadZone=0.000000,Sensitivity=1.000000,Exponent=1.000000,bInvert=False))
 +AxisConfig=(AxisKeyName="Gamepad_Special_Left_X",AxisProperties=(DeadZone=0.000000,Sensitivity=1.000000,Exponent=1.000000,bInvert=False))
 +AxisConfig=(AxisKeyName="Gamepad_Special_Left_Y",AxisProperties=(DeadZone=0.000000,Sensitivity=1.000000,Exponent=1.000000,bInvert=False))
++AxisConfig=(AxisKeyName="Daydream_Left_Trackpad_X",AxisProperties=(DeadZone=0.000000,Sensitivity=1.000000,Exponent=1.000000,bInvert=False))
++AxisConfig=(AxisKeyName="Daydream_Left_Trackpad_Y",AxisProperties=(DeadZone=0.000000,Sensitivity=1.000000,Exponent=1.000000,bInvert=False))
++AxisConfig=(AxisKeyName="Daydream_Right_Trackpad_X",AxisProperties=(DeadZone=0.000000,Sensitivity=1.000000,Exponent=1.000000,bInvert=False))
++AxisConfig=(AxisKeyName="Daydream_Right_Trackpad_Y",AxisProperties=(DeadZone=0.000000,Sensitivity=1.000000,Exponent=1.000000,bInvert=False))
++AxisConfig=(AxisKeyName="Vive_Left_Trigger_Axis",AxisProperties=(DeadZone=0.000000,Sensitivity=1.000000,Exponent=1.000000,bInvert=False))
++AxisConfig=(AxisKeyName="Vive_Left_Trackpad_X",AxisProperties=(DeadZone=0.000000,Sensitivity=1.000000,Exponent=1.000000,bInvert=False))
++AxisConfig=(AxisKeyName="Vive_Left_Trackpad_Y",AxisProperties=(DeadZone=0.000000,Sensitivity=1.000000,Exponent=1.000000,bInvert=False))
++AxisConfig=(AxisKeyName="Vive_Right_Trigger_Axis",AxisProperties=(DeadZone=0.000000,Sensitivity=1.000000,Exponent=1.000000,bInvert=False))
++AxisConfig=(AxisKeyName="Vive_Right_Trackpad_X",AxisProperties=(DeadZone=0.000000,Sensitivity=1.000000,Exponent=1.000000,bInvert=False))
++AxisConfig=(AxisKeyName="Vive_Right_Trackpad_Y",AxisProperties=(DeadZone=0.000000,Sensitivity=1.000000,Exponent=1.000000,bInvert=False))
++AxisConfig=(AxisKeyName="MixedReality_Left_Trigger_Axis",AxisProperties=(DeadZone=0.000000,Sensitivity=1.000000,Exponent=1.000000,bInvert=False))
++AxisConfig=(AxisKeyName="MixedReality_Left_Thumbstick_X",AxisProperties=(DeadZone=0.000000,Sensitivity=1.000000,Exponent=1.000000,bInvert=False))
++AxisConfig=(AxisKeyName="MixedReality_Left_Thumbstick_Y",AxisProperties=(DeadZone=0.000000,Sensitivity=1.000000,Exponent=1.000000,bInvert=False))
++AxisConfig=(AxisKeyName="MixedReality_Left_Trackpad_X",AxisProperties=(DeadZone=0.000000,Sensitivity=1.000000,Exponent=1.000000,bInvert=False))
++AxisConfig=(AxisKeyName="MixedReality_Left_Trackpad_Y",AxisProperties=(DeadZone=0.000000,Sensitivity=1.000000,Exponent=1.000000,bInvert=False))
++AxisConfig=(AxisKeyName="MixedReality_Right_Trigger_Axis",AxisProperties=(DeadZone=0.000000,Sensitivity=1.000000,Exponent=1.000000,bInvert=False))
++AxisConfig=(AxisKeyName="MixedReality_Right_Thumbstick_X",AxisProperties=(DeadZone=0.000000,Sensitivity=1.000000,Exponent=1.000000,bInvert=False))
++AxisConfig=(AxisKeyName="MixedReality_Right_Thumbstick_Y",AxisProperties=(DeadZone=0.000000,Sensitivity=1.000000,Exponent=1.000000,bInvert=False))
++AxisConfig=(AxisKeyName="MixedReality_Right_Trackpad_X",AxisProperties=(DeadZone=0.000000,Sensitivity=1.000000,Exponent=1.000000,bInvert=False))
++AxisConfig=(AxisKeyName="MixedReality_Right_Trackpad_Y",AxisProperties=(DeadZone=0.000000,Sensitivity=1.000000,Exponent=1.000000,bInvert=False))
++AxisConfig=(AxisKeyName="OculusGo_Left_Trackpad_X",AxisProperties=(DeadZone=0.000000,Sensitivity=1.000000,Exponent=1.000000,bInvert=False))
++AxisConfig=(AxisKeyName="OculusGo_Left_Trackpad_Y",AxisProperties=(DeadZone=0.000000,Sensitivity=1.000000,Exponent=1.000000,bInvert=False))
++AxisConfig=(AxisKeyName="OculusGo_Right_Trackpad_X",AxisProperties=(DeadZone=0.000000,Sensitivity=1.000000,Exponent=1.000000,bInvert=False))
++AxisConfig=(AxisKeyName="OculusGo_Right_Trackpad_Y",AxisProperties=(DeadZone=0.000000,Sensitivity=1.000000,Exponent=1.000000,bInvert=False))
++AxisConfig=(AxisKeyName="OculusTouch_Left_Grip_Axis",AxisProperties=(DeadZone=0.000000,Sensitivity=1.000000,Exponent=1.000000,bInvert=False))
++AxisConfig=(AxisKeyName="OculusTouch_Left_Trigger_Axis",AxisProperties=(DeadZone=0.000000,Sensitivity=1.000000,Exponent=1.000000,bInvert=False))
++AxisConfig=(AxisKeyName="OculusTouch_Left_Thumbstick_X",AxisProperties=(DeadZone=0.000000,Sensitivity=1.000000,Exponent=1.000000,bInvert=False))
++AxisConfig=(AxisKeyName="OculusTouch_Left_Thumbstick_Y",AxisProperties=(DeadZone=0.000000,Sensitivity=1.000000,Exponent=1.000000,bInvert=False))
++AxisConfig=(AxisKeyName="OculusTouch_Right_Grip_Axis",AxisProperties=(DeadZone=0.000000,Sensitivity=1.000000,Exponent=1.000000,bInvert=False))
++AxisConfig=(AxisKeyName="OculusTouch_Right_Trigger_Axis",AxisProperties=(DeadZone=0.000000,Sensitivity=1.000000,Exponent=1.000000,bInvert=False))
++AxisConfig=(AxisKeyName="OculusTouch_Right_Thumbstick_X",AxisProperties=(DeadZone=0.000000,Sensitivity=1.000000,Exponent=1.000000,bInvert=False))
++AxisConfig=(AxisKeyName="OculusTouch_Right_Thumbstick_Y",AxisProperties=(DeadZone=0.000000,Sensitivity=1.000000,Exponent=1.000000,bInvert=False))
++AxisConfig=(AxisKeyName="ValveIndex_Left_Grip_Axis",AxisProperties=(DeadZone=0.000000,Sensitivity=1.000000,Exponent=1.000000,bInvert=False))
++AxisConfig=(AxisKeyName="ValveIndex_Left_Grip_Force",AxisProperties=(DeadZone=0.000000,Sensitivity=1.000000,Exponent=1.000000,bInvert=False))
++AxisConfig=(AxisKeyName="ValveIndex_Left_Trigger_Axis",AxisProperties=(DeadZone=0.000000,Sensitivity=1.000000,Exponent=1.000000,bInvert=False))
++AxisConfig=(AxisKeyName="ValveIndex_Left_Thumbstick_X",AxisProperties=(DeadZone=0.000000,Sensitivity=1.000000,Exponent=1.000000,bInvert=False))
++AxisConfig=(AxisKeyName="ValveIndex_Left_Thumbstick_Y",AxisProperties=(DeadZone=0.000000,Sensitivity=1.000000,Exponent=1.000000,bInvert=False))
++AxisConfig=(AxisKeyName="ValveIndex_Left_Trackpad_X",AxisProperties=(DeadZone=0.000000,Sensitivity=1.000000,Exponent=1.000000,bInvert=False))
++AxisConfig=(AxisKeyName="ValveIndex_Left_Trackpad_Y",AxisProperties=(DeadZone=0.000000,Sensitivity=1.000000,Exponent=1.000000,bInvert=False))
++AxisConfig=(AxisKeyName="ValveIndex_Left_Trackpad_Force",AxisProperties=(DeadZone=0.000000,Sensitivity=1.000000,Exponent=1.000000,bInvert=False))
++AxisConfig=(AxisKeyName="ValveIndex_Left_Trackpad_Touch",AxisProperties=(DeadZone=0.000000,Sensitivity=1.000000,Exponent=1.000000,bInvert=False))
++AxisConfig=(AxisKeyName="ValveIndex_Right_Grip_Axis",AxisProperties=(DeadZone=0.000000,Sensitivity=1.000000,Exponent=1.000000,bInvert=False))
++AxisConfig=(AxisKeyName="ValveIndex_Right_Grip_Force",AxisProperties=(DeadZone=0.000000,Sensitivity=1.000000,Exponent=1.000000,bInvert=False))
++AxisConfig=(AxisKeyName="ValveIndex_Right_Trigger_Axis",AxisProperties=(DeadZone=0.000000,Sensitivity=1.000000,Exponent=1.000000,bInvert=False))
++AxisConfig=(AxisKeyName="ValveIndex_Right_Thumbstick_X",AxisProperties=(DeadZone=0.000000,Sensitivity=1.000000,Exponent=1.000000,bInvert=False))
++AxisConfig=(AxisKeyName="ValveIndex_Right_Thumbstick_Y",AxisProperties=(DeadZone=0.000000,Sensitivity=1.000000,Exponent=1.000000,bInvert=False))
++AxisConfig=(AxisKeyName="ValveIndex_Right_Trackpad_X",AxisProperties=(DeadZone=0.000000,Sensitivity=1.000000,Exponent=1.000000,bInvert=False))
++AxisConfig=(AxisKeyName="ValveIndex_Right_Trackpad_Y",AxisProperties=(DeadZone=0.000000,Sensitivity=1.000000,Exponent=1.000000,bInvert=False))
++AxisConfig=(AxisKeyName="ValveIndex_Right_Trackpad_Force",AxisProperties=(DeadZone=0.000000,Sensitivity=1.000000,Exponent=1.000000,bInvert=False))
++AxisConfig=(AxisKeyName="MotionController_Left_Thumbstick_Z",AxisProperties=(DeadZone=0.000000,Sensitivity=1.000000,Exponent=1.000000,bInvert=False))
++AxisConfig=(AxisKeyName="MagicLeap_Left_Trigger_Axis",AxisProperties=(DeadZone=0.000000,Sensitivity=1.000000,Exponent=1.000000,bInvert=False))
++AxisConfig=(AxisKeyName="MagicLeap_Left_Trackpad_X",AxisProperties=(DeadZone=0.000000,Sensitivity=1.000000,Exponent=1.000000,bInvert=False))
++AxisConfig=(AxisKeyName="MagicLeap_Left_Trackpad_Y",AxisProperties=(DeadZone=0.000000,Sensitivity=1.000000,Exponent=1.000000,bInvert=False))
++AxisConfig=(AxisKeyName="MagicLeap_Left_Trackpad_Force",AxisProperties=(DeadZone=0.000000,Sensitivity=1.000000,Exponent=1.000000,bInvert=False))
++AxisConfig=(AxisKeyName="MagicLeap_Left_Touch1_X",AxisProperties=(DeadZone=0.000000,Sensitivity=1.000000,Exponent=1.000000,bInvert=False))
++AxisConfig=(AxisKeyName="MagicLeap_Left_Touch1_Y",AxisProperties=(DeadZone=0.000000,Sensitivity=1.000000,Exponent=1.000000,bInvert=False))
++AxisConfig=(AxisKeyName="MagicLeap_Left_Touch1_Force",AxisProperties=(DeadZone=0.000000,Sensitivity=1.000000,Exponent=1.000000,bInvert=False))
++AxisConfig=(AxisKeyName="MotionController_Right_Thumbstick_Z",AxisProperties=(DeadZone=0.000000,Sensitivity=1.000000,Exponent=1.000000,bInvert=False))
++AxisConfig=(AxisKeyName="MagicLeap_Right_Trigger_Axis",AxisProperties=(DeadZone=0.000000,Sensitivity=1.000000,Exponent=1.000000,bInvert=False))
++AxisConfig=(AxisKeyName="MagicLeap_Right_Trackpad_X",AxisProperties=(DeadZone=0.000000,Sensitivity=1.000000,Exponent=1.000000,bInvert=False))
++AxisConfig=(AxisKeyName="MagicLeap_Right_Trackpad_Y",AxisProperties=(DeadZone=0.000000,Sensitivity=1.000000,Exponent=1.000000,bInvert=False))
++AxisConfig=(AxisKeyName="MagicLeap_Right_Trackpad_Force",AxisProperties=(DeadZone=0.000000,Sensitivity=1.000000,Exponent=1.000000,bInvert=False))
++AxisConfig=(AxisKeyName="MagicLeap_Right_Touch1_X",AxisProperties=(DeadZone=0.000000,Sensitivity=1.000000,Exponent=1.000000,bInvert=False))
++AxisConfig=(AxisKeyName="MagicLeap_Right_Touch1_Y",AxisProperties=(DeadZone=0.000000,Sensitivity=1.000000,Exponent=1.000000,bInvert=False))
++AxisConfig=(AxisKeyName="MagicLeap_Right_Touch1_Force",AxisProperties=(DeadZone=0.000000,Sensitivity=1.000000,Exponent=1.000000,bInvert=False))
++AxisConfig=(AxisKeyName="OculusTouch_Left_Thumbstick",AxisProperties=(DeadZone=0.000000,Sensitivity=1.000000,Exponent=1.000000,bInvert=False))
++AxisConfig=(AxisKeyName="OculusTouch_Left_FaceButton1",AxisProperties=(DeadZone=0.000000,Sensitivity=1.000000,Exponent=1.000000,bInvert=False))
++AxisConfig=(AxisKeyName="OculusTouch_Left_Trigger",AxisProperties=(DeadZone=0.000000,Sensitivity=1.000000,Exponent=1.000000,bInvert=False))
++AxisConfig=(AxisKeyName="OculusTouch_Left_FaceButton2",AxisProperties=(DeadZone=0.000000,Sensitivity=1.000000,Exponent=1.000000,bInvert=False))
++AxisConfig=(AxisKeyName="OculusTouch_Left_IndexPointing",AxisProperties=(DeadZone=0.000000,Sensitivity=1.000000,Exponent=1.000000,bInvert=False))
++AxisConfig=(AxisKeyName="OculusTouch_Left_ThumbUp",AxisProperties=(DeadZone=0.000000,Sensitivity=1.000000,Exponent=1.000000,bInvert=False))
++AxisConfig=(AxisKeyName="OculusTouch_Right_Thumbstick",AxisProperties=(DeadZone=0.000000,Sensitivity=1.000000,Exponent=1.000000,bInvert=False))
++AxisConfig=(AxisKeyName="OculusTouch_Right_FaceButton1",AxisProperties=(DeadZone=0.000000,Sensitivity=1.000000,Exponent=1.000000,bInvert=False))
++AxisConfig=(AxisKeyName="OculusTouch_Right_Trigger",AxisProperties=(DeadZone=0.000000,Sensitivity=1.000000,Exponent=1.000000,bInvert=False))
++AxisConfig=(AxisKeyName="OculusTouch_Right_FaceButton2",AxisProperties=(DeadZone=0.000000,Sensitivity=1.000000,Exponent=1.000000,bInvert=False))
++AxisConfig=(AxisKeyName="OculusTouch_Right_IndexPointing",AxisProperties=(DeadZone=0.000000,Sensitivity=1.000000,Exponent=1.000000,bInvert=False))
++AxisConfig=(AxisKeyName="OculusTouch_Right_ThumbUp",AxisProperties=(DeadZone=0.000000,Sensitivity=1.000000,Exponent=1.000000,bInvert=False))
++AxisConfig=(AxisKeyName="OculusTouchpad_Touchpad_X",AxisProperties=(DeadZone=0.000000,Sensitivity=1.000000,Exponent=1.000000,bInvert=False))
++AxisConfig=(AxisKeyName="OculusTouchpad_Touchpad_Y",AxisProperties=(DeadZone=0.000000,Sensitivity=1.000000,Exponent=1.000000,bInvert=False))
 bAltEnterTogglesFullscreen=True
+bF11TogglesFullscreen=True
 bUseMouseForTouch=False
 bEnableMouseSmoothing=True
 bEnableFOVScaling=True
-FOVScale=0.011110
-DoubleClickTime=0.200000
 bCaptureMouseOnLaunch=True
-DefaultViewportMouseCaptureMode=CapturePermanently_IncludingInitialMouseDown
-bDefaultViewportMouseLock=True
-+ActionMappings=(ActionName="Handbrake",Key=SpaceBar,bShift=False,bCtrl=False,bAlt=False,bCmd=False)
-+ActionMappings=(ActionName="Handbrake",Key=Gamepad_LeftShoulder,bShift=False,bCtrl=False,bAlt=False,bCmd=False)
-+ActionMappings=(ActionName="SwitchCamera",Key=Tab,bShift=False,bCtrl=False,bAlt=False,bCmd=False)
-+ActionMappings=(ActionName="ResetVR",Key=R,bShift=False,bCtrl=False,bAlt=False,bCmd=False)
-+ActionMappings=(ActionName="SwitchCamera",Key=Gamepad_Special_Left,bShift=False,bCtrl=False,bAlt=False,bCmd=False)
-+ActionMappings=(ActionName="Handbrake",Key=MotionController_Right_Grip1,bShift=False,bCtrl=False,bAlt=False,bCmd=False)
-+ActionMappings=(ActionName="ResetVR",Key=MotionController_Left_Grip1,bShift=False,bCtrl=False,bAlt=False,bCmd=False)
-+AxisMappings=(AxisName="MoveForward",Key=W,Scale=1.000000)
-+AxisMappings=(AxisName="MoveForward",Key=S,Scale=-1.000000)
-+AxisMappings=(AxisName="MoveForward",Key=Up,Scale=1.000000)
-+AxisMappings=(AxisName="MoveForward",Key=Down,Scale=-1.000000)
-+AxisMappings=(AxisName="MoveForward",Key=Gamepad_FaceButton_Bottom,Scale=1.000000)
-+AxisMappings=(AxisName="MoveForward",Key=Gamepad_RightY,Scale=-1.000000)
-+AxisMappings=(AxisName="MoveRight",Key=A,Scale=-1.000000)
-+AxisMappings=(AxisName="MoveRight",Key=D,Scale=1.000000)
-+AxisMappings=(AxisName="MoveRight",Key=Left,Scale=-1.000000)
-+AxisMappings=(AxisName="MoveRight",Key=Right,Scale=1.000000)
-+AxisMappings=(AxisName="MoveRight",Key=Gamepad_LeftX,Scale=1.000000)
-+AxisMappings=(AxisName="LookUp",Key=MouseY,Scale=1.000000)
-+AxisMappings=(AxisName="LookRight",Key=MouseX,Scale=1.000000)
-+AxisMappings=(AxisName="MoveForward",Key=Gamepad_RightTriggerAxis,Scale=1.000000)
-+AxisMappings=(AxisName="MoveForward",Key=Gamepad_LeftTriggerAxis,Scale=-1.000000)
-+AxisMappings=(AxisName="MoveForward",Key=MotionController_Right_TriggerAxis,Scale=1.000000)
-+AxisMappings=(AxisName="MoveForward",Key=MotionController_Left_TriggerAxis,Scale=-1.000000)
-+AxisMappings=(AxisName="MoveRight",Key=MotionController_Left_Thumbstick_X,Scale=1.000000)
-+AxisMappings=(AxisName="MoveRight",Key=MotionController_Right_Thumbstick_X,Scale=1.000000)
-+AxisMappings=(AxisName="MoveForward",Key=GenericUSBController_Axis2,Scale=-2.000000)
-+AxisMappings=(AxisName="MoveForward",Key=GenericUSBController_Axis3,Scale=1.000000)
-+AxisMappings=(AxisName="MoveRight",Key=GenericUSBController_Axis1,Scale=3.000000)
 bAlwaysShowTouchInterface=False
 bShowConsoleOnFourFingerTap=True
+bEnableGestureRecognizer=False
+bUseAutocorrect=False
+DefaultViewportMouseCaptureMode=CapturePermanently_IncludingInitialMouseDown
+DefaultViewportMouseLockMode=LockOnCapture
+FOVScale=0.011110
+DoubleClickTime=0.200000
++ActionMappings=(ActionName="Handbrake",bShift=False,bCtrl=False,bAlt=False,bCmd=False,Key=SpaceBar)
++ActionMappings=(ActionName="Handbrake",bShift=False,bCtrl=False,bAlt=False,bCmd=False,Key=Gamepad_LeftShoulder)
++ActionMappings=(ActionName="SwitchCamera",bShift=False,bCtrl=False,bAlt=False,bCmd=False,Key=Tab)
++ActionMappings=(ActionName="ResetVR",bShift=False,bCtrl=False,bAlt=False,bCmd=False,Key=R)
++ActionMappings=(ActionName="SwitchCamera",bShift=False,bCtrl=False,bAlt=False,bCmd=False,Key=Gamepad_Special_Left)
++ActionMappings=(ActionName="ResetVR",bShift=False,bCtrl=False,bAlt=False,bCmd=False,Key=Vive_Left_Menu_Click)
++ActionMappings=(ActionName="ResetVR",bShift=False,bCtrl=False,bAlt=False,bCmd=False,Key=MixedReality_Left_Thumbstick_Click)
++ActionMappings=(ActionName="ResetVR",bShift=False,bCtrl=False,bAlt=False,bCmd=False,Key=OculusTouch_Left_Thumbstick_Click)
++ActionMappings=(ActionName="ResetVR",bShift=False,bCtrl=False,bAlt=False,bCmd=False,Key=ValveIndex_Left_Thumbstick_Click)
++ActionMappings=(ActionName="ResetVR",bShift=False,bCtrl=False,bAlt=False,bCmd=False,Key=MagicLeap_Left_Bumper)
++ActionMappings=(ActionName="Handbrake",bShift=False,bCtrl=False,bAlt=False,bCmd=False,Key=Vive_Right_Grip_Click)
++ActionMappings=(ActionName="Handbrake",bShift=False,bCtrl=False,bAlt=False,bCmd=False,Key=MixedReality_Right_Grip_Click)
++ActionMappings=(ActionName="Handbrake",bShift=False,bCtrl=False,bAlt=False,bCmd=False,Key=OculusTouch_Right_Grip_Click)
++ActionMappings=(ActionName="Handbrake",bShift=False,bCtrl=False,bAlt=False,bCmd=False,Key=ValveIndex_Right_A_Click)
++ActionMappings=(ActionName="Handbrake",bShift=False,bCtrl=False,bAlt=False,bCmd=False,Key=MagicLeap_Right_Bumper)
++AxisMappings=(AxisName="MoveForward",Scale=1.000000,Key=W)
++AxisMappings=(AxisName="MoveForward",Scale=-1.000000,Key=S)
++AxisMappings=(AxisName="MoveForward",Scale=1.000000,Key=Up)
++AxisMappings=(AxisName="MoveForward",Scale=-1.000000,Key=Down)
++AxisMappings=(AxisName="MoveForward",Scale=1.000000,Key=Gamepad_FaceButton_Bottom)
++AxisMappings=(AxisName="MoveForward",Scale=-1.000000,Key=Gamepad_RightY)
++AxisMappings=(AxisName="MoveRight",Scale=-1.000000,Key=A)
++AxisMappings=(AxisName="MoveRight",Scale=1.000000,Key=D)
++AxisMappings=(AxisName="MoveRight",Scale=-1.000000,Key=Left)
++AxisMappings=(AxisName="MoveRight",Scale=1.000000,Key=Right)
++AxisMappings=(AxisName="MoveRight",Scale=1.000000,Key=Gamepad_LeftX)
++AxisMappings=(AxisName="LookUp",Scale=1.000000,Key=MouseY)
++AxisMappings=(AxisName="LookRight",Scale=1.000000,Key=MouseX)
++AxisMappings=(AxisName="MoveForward",Scale=1.000000,Key=Gamepad_RightTriggerAxis)
++AxisMappings=(AxisName="MoveForward",Scale=-1.000000,Key=Gamepad_LeftTriggerAxis)
++AxisMappings=(AxisName="MoveForward",Scale=-2.000000,Key=GenericUSBController_Axis2)
++AxisMappings=(AxisName="MoveForward",Scale=1.000000,Key=GenericUSBController_Axis3)
++AxisMappings=(AxisName="MoveRight",Scale=3.000000,Key=GenericUSBController_Axis1)
++AxisMappings=(AxisName="MoveForward",Scale=1.000000,Key=Vive_Right_Trigger_Axis)
++AxisMappings=(AxisName="MoveForward",Scale=-1.000000,Key=Vive_Left_Trigger_Axis)
++AxisMappings=(AxisName="MoveForward",Scale=1.000000,Key=MixedReality_Right_Trigger_Axis)
++AxisMappings=(AxisName="MoveForward",Scale=-1.000000,Key=MixedReality_Left_Trigger_Axis)
++AxisMappings=(AxisName="MoveForward",Scale=1.000000,Key=OculusTouch_Right_Trigger_Axis)
++AxisMappings=(AxisName="MoveForward",Scale=-1.000000,Key=OculusTouch_Left_Trigger_Axis)
++AxisMappings=(AxisName="MoveForward",Scale=1.000000,Key=ValveIndex_Right_Trigger_Axis)
++AxisMappings=(AxisName="MoveForward",Scale=-1.000000,Key=ValveIndex_Left_Trigger_Axis)
++AxisMappings=(AxisName="MoveForward",Scale=1.000000,Key=MagicLeap_Right_Trigger_Axis)
++AxisMappings=(AxisName="MoveForward",Scale=-1.000000,Key=MagicLeap_Left_Trigger_Axis)
++AxisMappings=(AxisName="MoveRight",Scale=1.000000,Key=Vive_Left_Trackpad_X)
++AxisMappings=(AxisName="MoveRight",Scale=1.000000,Key=Vive_Right_Trackpad_X)
++AxisMappings=(AxisName="MoveRight",Scale=1.000000,Key=MixedReality_Left_Thumbstick_X)
++AxisMappings=(AxisName="MoveRight",Scale=1.000000,Key=MixedReality_Right_Thumbstick_X)
++AxisMappings=(AxisName="MoveRight",Scale=1.000000,Key=OculusTouch_Left_Thumbstick_X)
++AxisMappings=(AxisName="MoveRight",Scale=1.000000,Key=OculusTouch_Right_Thumbstick_X)
++AxisMappings=(AxisName="MoveRight",Scale=1.000000,Key=ValveIndex_Left_Thumbstick_X)
++AxisMappings=(AxisName="MoveRight",Scale=1.000000,Key=ValveIndex_Right_Thumbstick_X)
++AxisMappings=(AxisName="MoveRight",Scale=1.000000,Key=MagicLeap_Left_Trackpad_X)
++AxisMappings=(AxisName="MoveRight",Scale=1.000000,Key=MagicLeap_Right_Trackpad_X)
 DefaultTouchInterface=/Game/Vehicle/Sedan/SedanVirtualJoysticks.SedanVirtualJoysticks
-ConsoleKey=None
--ConsoleKeys=Tilde
 +ConsoleKeys=Tilde
 
-[/Script/RawInput.RawInputSettings]
-+DeviceConfigurations=(VendorID="0x046D",ProductID="0xC262",AxisProperties=((bEnabled=False),(Key=GenericUSBController_Axis1,Offset=-0.500000),(bEnabled=False),(Key=GenericUSBController_Axis2,bInverted=True,Offset=1.000000),(Key=GenericUSBController_Axis3,bInverted=True,Offset=1.000000)),ButtonProperties=((Key=GenericUSBController_Button1),(Key=GenericUSBController_Button2),(Key=GenericUSBController_Button3),(Key=GenericUSBController_Button4),(Key=GenericUSBController_Button5),(Key=GenericUSBController_Button6),(Key=GenericUSBController_Button7),(Key=GenericUSBController_Button8),(Key=GenericUSBController_Button9),(Key=GenericUSBController_Button10),(Key=GenericUSBController_Button11),(Key=GenericUSBController_Button12)))
-
 
diff --git a/Templates/TP_VirtualRealityBP/Config/DefaultInput.ini b/Templates/TP_VirtualRealityBP/Config/DefaultInput.ini
--- a/Templates/TP_VirtualRealityBP/Config/DefaultInput.ini
+++ b/Templates/TP_VirtualRealityBP/Config/DefaultInput.ini
@@ -1,11 +1,5 @@
 
 [/Script/Engine.InputSettings]
--AxisConfig=(AxisKeyName="Gamepad_LeftX",AxisProperties=(DeadZone=0.25,Exponent=1.f,Sensitivity=1.f))
--AxisConfig=(AxisKeyName="Gamepad_LeftY",AxisProperties=(DeadZone=0.25,Exponent=1.f,Sensitivity=1.f))
--AxisConfig=(AxisKeyName="Gamepad_RightX",AxisProperties=(DeadZone=0.25,Exponent=1.f,Sensitivity=1.f))
--AxisConfig=(AxisKeyName="Gamepad_RightY",AxisProperties=(DeadZone=0.25,Exponent=1.f,Sensitivity=1.f))
--AxisConfig=(AxisKeyName="MouseX",AxisProperties=(DeadZone=0.f,Exponent=1.f,Sensitivity=0.07f))
--AxisConfig=(AxisKeyName="MouseY",AxisProperties=(DeadZone=0.f,Exponent=1.f,Sensitivity=0.07f))
 +AxisConfig=(AxisKeyName="OculusTouch_Right_FaceButton2",AxisProperties=(DeadZone=0.000000,Sensitivity=1.000000,Exponent=1.000000,bInvert=False))
 +AxisConfig=(AxisKeyName="OculusTouch_Right_IndexPointing",AxisProperties=(DeadZone=0.000000,Sensitivity=1.000000,Exponent=1.000000,bInvert=False))
 +AxisConfig=(AxisKeyName="OculusTouch_Right_ThumbUp",AxisProperties=(DeadZone=0.000000,Sensitivity=1.000000,Exponent=1.000000,bInvert=False))
@@ -88,13 +82,30 @@
 +AxisConfig=(AxisKeyName="ValveIndex_Right_Trackpad_X",AxisProperties=(DeadZone=0.000000,Sensitivity=1.000000,Exponent=1.000000,bInvert=False))
 +AxisConfig=(AxisKeyName="ValveIndex_Right_Trackpad_Y",AxisProperties=(DeadZone=0.000000,Sensitivity=1.000000,Exponent=1.000000,bInvert=False))
 +AxisConfig=(AxisKeyName="ValveIndex_Right_Trackpad_Force",AxisProperties=(DeadZone=0.000000,Sensitivity=1.000000,Exponent=1.000000,bInvert=False))
++AxisConfig=(AxisKeyName="MotionController_Left_Thumbstick_Z",AxisProperties=(DeadZone=0.000000,Sensitivity=1.000000,Exponent=1.000000,bInvert=False))
++AxisConfig=(AxisKeyName="MagicLeap_Left_Trigger_Axis",AxisProperties=(DeadZone=0.000000,Sensitivity=1.000000,Exponent=1.000000,bInvert=False))
++AxisConfig=(AxisKeyName="MagicLeap_Left_Trackpad_X",AxisProperties=(DeadZone=0.000000,Sensitivity=1.000000,Exponent=1.000000,bInvert=False))
++AxisConfig=(AxisKeyName="MagicLeap_Left_Trackpad_Y",AxisProperties=(DeadZone=0.000000,Sensitivity=1.000000,Exponent=1.000000,bInvert=False))
++AxisConfig=(AxisKeyName="MagicLeap_Left_Trackpad_Force",AxisProperties=(DeadZone=0.000000,Sensitivity=1.000000,Exponent=1.000000,bInvert=False))
++AxisConfig=(AxisKeyName="MagicLeap_Left_Touch1_X",AxisProperties=(DeadZone=0.000000,Sensitivity=1.000000,Exponent=1.000000,bInvert=False))
++AxisConfig=(AxisKeyName="MagicLeap_Left_Touch1_Y",AxisProperties=(DeadZone=0.000000,Sensitivity=1.000000,Exponent=1.000000,bInvert=False))
++AxisConfig=(AxisKeyName="MagicLeap_Left_Touch1_Force",AxisProperties=(DeadZone=0.000000,Sensitivity=1.000000,Exponent=1.000000,bInvert=False))
++AxisConfig=(AxisKeyName="MotionController_Right_Thumbstick_Z",AxisProperties=(DeadZone=0.000000,Sensitivity=1.000000,Exponent=1.000000,bInvert=False))
++AxisConfig=(AxisKeyName="MagicLeap_Right_Trigger_Axis",AxisProperties=(DeadZone=0.000000,Sensitivity=1.000000,Exponent=1.000000,bInvert=False))
++AxisConfig=(AxisKeyName="MagicLeap_Right_Trackpad_X",AxisProperties=(DeadZone=0.000000,Sensitivity=1.000000,Exponent=1.000000,bInvert=False))
++AxisConfig=(AxisKeyName="MagicLeap_Right_Trackpad_Y",AxisProperties=(DeadZone=0.000000,Sensitivity=1.000000,Exponent=1.000000,bInvert=False))
++AxisConfig=(AxisKeyName="MagicLeap_Right_Trackpad_Force",AxisProperties=(DeadZone=0.000000,Sensitivity=1.000000,Exponent=1.000000,bInvert=False))
++AxisConfig=(AxisKeyName="MagicLeap_Right_Touch1_X",AxisProperties=(DeadZone=0.000000,Sensitivity=1.000000,Exponent=1.000000,bInvert=False))
++AxisConfig=(AxisKeyName="MagicLeap_Right_Touch1_Y",AxisProperties=(DeadZone=0.000000,Sensitivity=1.000000,Exponent=1.000000,bInvert=False))
++AxisConfig=(AxisKeyName="MagicLeap_Right_Touch1_Force",AxisProperties=(DeadZone=0.000000,Sensitivity=1.000000,Exponent=1.000000,bInvert=False))
++AxisConfig=(AxisKeyName="OculusTouchpad_Touchpad_X",AxisProperties=(DeadZone=0.000000,Sensitivity=1.000000,Exponent=1.000000,bInvert=False))
++AxisConfig=(AxisKeyName="OculusTouchpad_Touchpad_Y",AxisProperties=(DeadZone=0.000000,Sensitivity=1.000000,Exponent=1.000000,bInvert=False))
 bAltEnterTogglesFullscreen=True
 bF11TogglesFullscreen=True
 bUseMouseForTouch=False
 bEnableMouseSmoothing=True
 bEnableFOVScaling=True
 bCaptureMouseOnLaunch=True
-bDefaultViewportMouseLock=True
 bAlwaysShowTouchInterface=False
 bShowConsoleOnFourFingerTap=True
 bEnableGestureRecognizer=False
@@ -133,6 +144,10 @@ DoubleClickTime=0.200000
 +ActionMappings=(ActionName="TeleportLeft",bShift=False,bCtrl=False,bAlt=False,bCmd=False,Key=ValveIndex_Left_A_Click)
 +ActionMappings=(ActionName="GrabLeft",bShift=False,bCtrl=False,bAlt=False,bCmd=False,Key=Daydream_Left_Select_Click)
 +ActionMappings=(ActionName="GrabRight",bShift=False,bCtrl=False,bAlt=False,bCmd=False,Key=Daydream_Right_Select_Click)
++ActionMappings=(ActionName="GrabLeft",bShift=False,bCtrl=False,bAlt=False,bCmd=False,Key=MagicLeap_Left_Trigger)
++ActionMappings=(ActionName="GrabRight",bShift=False,bCtrl=False,bAlt=False,bCmd=False,Key=MagicLeap_Right_Trigger)
++ActionMappings=(ActionName="TeleportRight",bShift=False,bCtrl=False,bAlt=False,bCmd=False,Key=MagicLeap_Right_Bumper)
++ActionMappings=(ActionName="TeleportLeft",bShift=False,bCtrl=False,bAlt=False,bCmd=False,Key=MagicLeap_Left_Bumper)
 +AxisMappings=(AxisName="TeleportDirectionUp",Scale=1.000000,Key=Gamepad_LeftY)
 +AxisMappings=(AxisName="TeleportDirectionRight",Scale=1.000000,Key=Gamepad_LeftX)
 +AxisMappings=(AxisName="MotionControllerThumbLeft_Y",Scale=1.000000,Key=Vive_Left_Trackpad_Y)
@@ -159,9 +174,11 @@ DoubleClickTime=0.200000
 +AxisMappings=(AxisName="MotionControllerThumbLeft_X",Scale=1.000000,Key=Daydream_Left_Trackpad_X)
 +AxisMappings=(AxisName="MotionControllerThumbRight_Y",Scale=1.000000,Key=Daydream_Right_Trackpad_Y)
 +AxisMappings=(AxisName="MotionControllerThumbRight_X",Scale=1.000000,Key=Daydream_Right_Trackpad_X)
++AxisMappings=(AxisName="MotionControllerThumbRight_Y",Scale=1.000000,Key=MagicLeap_Right_Trackpad_Y)
++AxisMappings=(AxisName="MotionControllerThumbRight_X",Scale=1.000000,Key=MagicLeap_Right_Trackpad_X)
++AxisMappings=(AxisName="MotionControllerThumbLeft_Y",Scale=1.000000,Key=MagicLeap_Left_Trackpad_Y)
++AxisMappings=(AxisName="MotionControllerThumbLeft_X",Scale=1.000000,Key=MagicLeap_Left_Trackpad_X)
 DefaultTouchInterface=None
-ConsoleKey=None
--ConsoleKeys=Tilde
 +ConsoleKeys=Tilde
 
 
